Move UMG/FieldNotification files to the FieldNotification module.

diff --git a/Engine/Config/BaseEngine.ini b/Engine/Config/BaseEngine.ini
--- a/Engine/Config/BaseEngine.ini
+++ b/Engine/Config/BaseEngine.ini
@@ -1137,9 +1137,12 @@ DefaultWorldPartitionHLODLayer=/Engine/Maps/Templates/HLODs/HLODLayer_Instanced.
 
 +PropertyRedirects=(OldName="/Script/Engine.SoundAttenuationSettings.OmniRadius", NewName="NonSpatializedRadiusStart")
 
-; 5.3
+;5.3
 +FunctionRedirects=(OldName="KismetMathLibrary.BreakVector3f", NewName="BreakVector")
 
++StructRedirects=(OldName="/Script/UMG.FieldNotificationId", NewName="/Script/FieldNotification.FieldNotificationId")
++ClassRedirects=(OldName="/Script/UMG.NotifyFieldValueChanged", NewName="/Script/FieldNotification.NotifyFieldValueChanged")
+
 [CoreUObject.Metadata]
 ; Note: UnrealHeaderTool should be rerun after making changes to MetadataRedirects to catch any keys specified in class headers
 +MetadataRedirects=(OldKey="K2Protected", NewKey="BlueprintProtected")
diff --git a/Engine/Plugins/Experimental/UIFramework/Source/UIFramework.Build.cs b/Engine/Plugins/Experimental/UIFramework/Source/UIFramework.Build.cs
--- a/Engine/Plugins/Experimental/UIFramework/Source/UIFramework.Build.cs
+++ b/Engine/Plugins/Experimental/UIFramework/Source/UIFramework.Build.cs
@@ -13,6 +13,7 @@ public class UIFramework : ModuleRules
 				"CoreUObject",
 				"DeveloperSettings",
 				"Engine",
+				"FieldNotification",
 				"ModelViewViewModel",
 				"SlateCore",
 				"Slate",
diff --git a/Engine/Plugins/Runtime/ModelViewViewModel/Source/ModelViewViewModel/ModelViewViewModel.Build.cs b/Engine/Plugins/Runtime/ModelViewViewModel/Source/ModelViewViewModel/ModelViewViewModel.Build.cs
--- a/Engine/Plugins/Runtime/ModelViewViewModel/Source/ModelViewViewModel/ModelViewViewModel.Build.cs
+++ b/Engine/Plugins/Runtime/ModelViewViewModel/Source/ModelViewViewModel/ModelViewViewModel.Build.cs
@@ -14,6 +14,7 @@ public class ModelViewViewModel : ModuleRules
 				"Core",
 				"CoreUObject",
 				"Engine",
+				"FieldNotification",
 			});
 
 		PrivateDependencyModuleNames.AddRange(
diff --git a/Engine/Plugins/Runtime/ModelViewViewModel/Source/ModelViewViewModel/Private/Bindings/MVVMCompiledBindingLibrary.cpp b/Engine/Plugins/Runtime/ModelViewViewModel/Source/ModelViewViewModel/Private/Bindings/MVVMCompiledBindingLibrary.cpp
--- a/Engine/Plugins/Runtime/ModelViewViewModel/Source/ModelViewViewModel/Private/Bindings/MVVMCompiledBindingLibrary.cpp
+++ b/Engine/Plugins/Runtime/ModelViewViewModel/Source/ModelViewViewModel/Private/Bindings/MVVMCompiledBindingLibrary.cpp
@@ -3,12 +3,12 @@
 #include "Bindings/MVVMCompiledBindingLibrary.h"
 
 #include "Bindings/MVVMBindingHelper.h"
-#include "FieldNotification/IClassDescriptor.h"
-#include "UObject/PropertyAccessUtil.h"
-#include "FieldNotification/IFieldValueChanged.h"
+#include "IFieldNotificationClassDescriptor.h"
+#include "INotifyFieldValueChanged.h"
 #include "Templates/ValueOrError.h"
 #include "Types/MVVMFieldContext.h"
 #include "Types/MVVMFunctionContext.h"
+#include "UObject/PropertyAccessUtil.h"
 #include "UObject/UnrealType.h"
 
 #include UE_INLINE_GENERATED_CPP_BY_NAME(MVVMCompiledBindingLibrary)
diff --git a/Engine/Plugins/Runtime/ModelViewViewModel/Source/ModelViewViewModel/Private/View/MVVMView.cpp b/Engine/Plugins/Runtime/ModelViewViewModel/Source/ModelViewViewModel/Private/View/MVVMView.cpp
--- a/Engine/Plugins/Runtime/ModelViewViewModel/Source/ModelViewViewModel/Private/View/MVVMView.cpp
+++ b/Engine/Plugins/Runtime/ModelViewViewModel/Source/ModelViewViewModel/Private/View/MVVMView.cpp
@@ -1,7 +1,7 @@
 // Copyright Epic Games, Inc. All Rights Reserved.
 
 #include "View/MVVMView.h"
-#include "FieldNotification/FieldMulticastDelegate.h"
+#include "FieldNotificationDelegate.h"
 #include "View/MVVMViewClass.h"
 #include "Misc/MemStack.h"
 #include "View/MVVMBindingSubsystem.h"
diff --git a/Engine/Plugins/Runtime/ModelViewViewModel/Source/ModelViewViewModel/Public/Bindings/MVVMCompiledBindingLibrary.h b/Engine/Plugins/Runtime/ModelViewViewModel/Source/ModelViewViewModel/Public/Bindings/MVVMCompiledBindingLibrary.h
--- a/Engine/Plugins/Runtime/ModelViewViewModel/Source/ModelViewViewModel/Public/Bindings/MVVMCompiledBindingLibrary.h
+++ b/Engine/Plugins/Runtime/ModelViewViewModel/Source/ModelViewViewModel/Public/Bindings/MVVMCompiledBindingLibrary.h
@@ -2,7 +2,7 @@
 
 #pragma once
 
-#include "FieldNotification/FieldId.h" // IWYU pragma: keep
+#include "FieldNotificationId.h" // IWYU pragma: keep
 #include "MVVMCompiledBindingLibrary.generated.h"
 
 namespace UE::FieldNotification { struct FFieldId; }
diff --git a/Engine/Plugins/Runtime/ModelViewViewModel/Source/ModelViewViewModel/Public/Bindings/MVVMFieldPathHelper.h b/Engine/Plugins/Runtime/ModelViewViewModel/Source/ModelViewViewModel/Public/Bindings/MVVMFieldPathHelper.h
--- a/Engine/Plugins/Runtime/ModelViewViewModel/Source/ModelViewViewModel/Public/Bindings/MVVMFieldPathHelper.h
+++ b/Engine/Plugins/Runtime/ModelViewViewModel/Source/ModelViewViewModel/Public/Bindings/MVVMFieldPathHelper.h
@@ -2,7 +2,7 @@
 
 #pragma once
 
-#include "FieldNotification/FieldId.h"
+#include "FieldNotificationId.h"
 #include "Templates/SubclassOf.h"
 
 namespace UE::MVVM { struct FFieldContext; }
diff --git a/Engine/Plugins/Runtime/ModelViewViewModel/Source/ModelViewViewModel/Public/MVVMViewModelBase.h b/Engine/Plugins/Runtime/ModelViewViewModel/Source/ModelViewViewModel/Public/MVVMViewModelBase.h
--- a/Engine/Plugins/Runtime/ModelViewViewModel/Source/ModelViewViewModel/Public/MVVMViewModelBase.h
+++ b/Engine/Plugins/Runtime/ModelViewViewModel/Source/ModelViewViewModel/Public/MVVMViewModelBase.h
@@ -3,8 +3,8 @@
 #pragma once
 
 
-#include "FieldNotification/IClassDescriptor.h"
-#include "FieldNotification/FieldNotificationDeclaration.h" // IWYU pragma: keep
+#include "FieldNotificationDeclaration.h" // IWYU pragma: keep
+#include "IFieldNotificationClassDescriptor.h"
 #include "ViewModel/MVVMFieldNotificationDelegates.h"
 
 #include "MVVMViewModelBase.generated.h"
diff --git a/Engine/Plugins/Runtime/ModelViewViewModel/Source/ModelViewViewModel/Public/View/MVVMViewModelContextResolver.h b/Engine/Plugins/Runtime/ModelViewViewModel/Source/ModelViewViewModel/Public/View/MVVMViewModelContextResolver.h
--- a/Engine/Plugins/Runtime/ModelViewViewModel/Source/ModelViewViewModel/Public/View/MVVMViewModelContextResolver.h
+++ b/Engine/Plugins/Runtime/ModelViewViewModel/Source/ModelViewViewModel/Public/View/MVVMViewModelContextResolver.h
@@ -3,7 +3,7 @@
 #pragma once
 
 #include "UObject/Object.h"
-#include "FieldNotification/IFieldValueChanged.h"
+#include "INotifyFieldValueChanged.h"
 
 #include "MVVMViewModelContextResolver.generated.h"
 
diff --git a/Engine/Plugins/Runtime/ModelViewViewModel/Source/ModelViewViewModel/Public/ViewModel/MVVMFieldNotificationDelegates.h b/Engine/Plugins/Runtime/ModelViewViewModel/Source/ModelViewViewModel/Public/ViewModel/MVVMFieldNotificationDelegates.h
--- a/Engine/Plugins/Runtime/ModelViewViewModel/Source/ModelViewViewModel/Public/ViewModel/MVVMFieldNotificationDelegates.h
+++ b/Engine/Plugins/Runtime/ModelViewViewModel/Source/ModelViewViewModel/Public/ViewModel/MVVMFieldNotificationDelegates.h
@@ -2,7 +2,7 @@
 
 #pragma once
 
-#include "FieldNotification/FieldMulticastDelegate.h"
+#include "FieldNotificationDelegate.h"
 
 class FDelegateHandle;
 namespace UE::FieldNotification { struct FFieldId; }
diff --git a/Engine/Plugins/Runtime/ModelViewViewModel/Source/ModelViewViewModelBlueprint/ModelViewViewModelBlueprint.Build.cs b/Engine/Plugins/Runtime/ModelViewViewModel/Source/ModelViewViewModelBlueprint/ModelViewViewModelBlueprint.Build.cs
--- a/Engine/Plugins/Runtime/ModelViewViewModel/Source/ModelViewViewModelBlueprint/ModelViewViewModelBlueprint.Build.cs
+++ b/Engine/Plugins/Runtime/ModelViewViewModel/Source/ModelViewViewModelBlueprint/ModelViewViewModelBlueprint.Build.cs
@@ -15,6 +15,7 @@ public class ModelViewViewModelBlueprint : ModuleRules
 				"CoreUObject",
 				"DeveloperSettings",
 				"Engine",
+				"FieldNotification",
 				"ModelViewViewModel",
 			});
 
diff --git a/Engine/Plugins/Runtime/ModelViewViewModel/Source/ModelViewViewModelBlueprint/Private/ViewModel/MVVMViewModelBlueprintCompiler.cpp b/Engine/Plugins/Runtime/ModelViewViewModel/Source/ModelViewViewModelBlueprint/Private/ViewModel/MVVMViewModelBlueprintCompiler.cpp
--- a/Engine/Plugins/Runtime/ModelViewViewModel/Source/ModelViewViewModelBlueprint/Private/ViewModel/MVVMViewModelBlueprintCompiler.cpp
+++ b/Engine/Plugins/Runtime/ModelViewViewModel/Source/ModelViewViewModelBlueprint/Private/ViewModel/MVVMViewModelBlueprintCompiler.cpp
@@ -8,7 +8,7 @@
 #include "MVVMViewModelBase.h"
 
 #include "FieldNotification/CustomizationHelper.h"
-#include "FieldNotification/FieldNotificationHelpers.h"
+#include "FieldNotificationHelpers.h"
 #include "K2Node_FunctionEntry.h"
 #include "Kismet2/KismetReinstanceUtilities.h"
 #include "MVVMFunctionGraphHelper.h"
diff --git a/Engine/Plugins/Runtime/ModelViewViewModel/Source/ModelViewViewModelEditor/ModelViewViewModelEditor.Build.cs b/Engine/Plugins/Runtime/ModelViewViewModel/Source/ModelViewViewModelEditor/ModelViewViewModelEditor.Build.cs
--- a/Engine/Plugins/Runtime/ModelViewViewModel/Source/ModelViewViewModelEditor/ModelViewViewModelEditor.Build.cs
+++ b/Engine/Plugins/Runtime/ModelViewViewModel/Source/ModelViewViewModelEditor/ModelViewViewModelEditor.Build.cs
@@ -15,6 +15,7 @@ public class ModelViewViewModelEditor : ModuleRules
 				"Core",
 				"CoreUObject",
 				"EditorSubsystem",
+				"FieldNotification",
 				"ModelViewViewModel",
 				"ModelViewViewModelBlueprint",
 				"PropertyPath",
diff --git a/Engine/Source/Editor/UMGEditor/Private/Graph/UMGGraphPanelPinFactory.h b/Engine/Source/Editor/UMGEditor/Private/Graph/UMGGraphPanelPinFactory.h
--- a/Engine/Source/Editor/UMGEditor/Private/Graph/UMGGraphPanelPinFactory.h
+++ b/Engine/Source/Editor/UMGEditor/Private/Graph/UMGGraphPanelPinFactory.h
@@ -5,7 +5,7 @@
 #include "CoreMinimal.h"
 #include "EdGraphUtilities.h"
 #include "EdGraphSchema_K2.h"
-#include "FieldNotification/FieldId.h"
+#include "FieldNotificationId.h"
 #include "SGraphPin.h"
 #include "UObject/Class.h"
 #include "Widgets/SFieldNotificationGraphPin.h"
diff --git a/Engine/Source/Editor/UMGEditor/Private/WidgetBlueprintCompiler.cpp b/Engine/Source/Editor/UMGEditor/Private/WidgetBlueprintCompiler.cpp
--- a/Engine/Source/Editor/UMGEditor/Private/WidgetBlueprintCompiler.cpp
+++ b/Engine/Source/Editor/UMGEditor/Private/WidgetBlueprintCompiler.cpp
@@ -12,7 +12,7 @@
 #include "MovieScene.h"
 
 #include "FieldNotification/CustomizationHelper.h"
-#include "FieldNotification/FieldNotificationHelpers.h"
+#include "FieldNotificationHelpers.h"
 #include "Kismet2/Kismet2NameValidators.h"
 #include "Kismet2/KismetReinstanceUtilities.h"
 #include "Kismet2/BlueprintEditorUtils.h"
diff --git a/Engine/Source/Editor/UMGEditor/Private/Widgets/SFieldNotificationGraphPin.h b/Engine/Source/Editor/UMGEditor/Private/Widgets/SFieldNotificationGraphPin.h
--- a/Engine/Source/Editor/UMGEditor/Private/Widgets/SFieldNotificationGraphPin.h
+++ b/Engine/Source/Editor/UMGEditor/Private/Widgets/SFieldNotificationGraphPin.h
@@ -3,7 +3,7 @@
 #pragma once
 
 #include "Delegates/Delegate.h"
-#include "FieldNotification/FieldId.h"
+#include "FieldNotificationId.h"
 #include "SGraphPin.h"
 #include "Templates/SharedPointer.h"
 #include "Widgets/DeclarativeSyntaxSupport.h"
diff --git a/Engine/Source/Editor/UMGEditor/Private/Widgets/SFieldNotificationPicker.cpp b/Engine/Source/Editor/UMGEditor/Private/Widgets/SFieldNotificationPicker.cpp
--- a/Engine/Source/Editor/UMGEditor/Private/Widgets/SFieldNotificationPicker.cpp
+++ b/Engine/Source/Editor/UMGEditor/Private/Widgets/SFieldNotificationPicker.cpp
@@ -2,8 +2,8 @@
 
 #include "Widgets/SFieldNotificationPicker.h"
 
-#include "FieldNotification/IClassDescriptor.h"
-#include "FieldNotification/IFieldValueChanged.h"
+#include "IFieldNotificationClassDescriptor.h"
+#include "INotifyFieldValueChanged.h"
 #include "Layout/Children.h"
 #include "UObject/Class.h"
 #include "UObject/NameTypes.h"
diff --git a/Engine/Source/Editor/UMGEditor/Private/Widgets/SFieldNotificationPicker.h b/Engine/Source/Editor/UMGEditor/Private/Widgets/SFieldNotificationPicker.h
--- a/Engine/Source/Editor/UMGEditor/Private/Widgets/SFieldNotificationPicker.h
+++ b/Engine/Source/Editor/UMGEditor/Private/Widgets/SFieldNotificationPicker.h
@@ -7,7 +7,7 @@
 #include "Containers/Set.h"
 #include "Containers/SparseArray.h"
 #include "Delegates/Delegate.h"
-#include "FieldNotification/FieldId.h"
+#include "FieldNotificationId.h"
 #include "HAL/PlatformCrt.h"
 #include "Internationalization/Text.h"
 #include "Misc/Attribute.h"
diff --git a/Engine/Source/Editor/UMGEditor/UMGEditor.Build.cs b/Engine/Source/Editor/UMGEditor/UMGEditor.Build.cs
--- a/Engine/Source/Editor/UMGEditor/UMGEditor.Build.cs
+++ b/Engine/Source/Editor/UMGEditor/UMGEditor.Build.cs
@@ -14,11 +14,12 @@ public class UMGEditor : ModuleRules
 				"UMG",
 			});
 
-        PublicDependencyModuleNames.AddRange(
-            new string[]
-            {
-                "Sequencer",
-            });
+		PublicDependencyModuleNames.AddRange(
+			new string[]
+			{
+				"FieldNotification",
+				"Sequencer",
+			});
 
 		PrivateDependencyModuleNames.AddRange(
 			new string[] {
diff --git a/Engine/Source/Runtime/FieldNotification/FieldNotification.Build.cs b/Engine/Source/Runtime/FieldNotification/FieldNotification.Build.cs
new file mode 100644
--- /dev/null
+++ b/Engine/Source/Runtime/FieldNotification/FieldNotification.Build.cs
@@ -0,0 +1,16 @@
+// Copyright Epic Games, Inc. All Rights Reserved.
+
+using UnrealBuildTool;
+
+public class FieldNotification: ModuleRules
+{
+	public FieldNotification(ReadOnlyTargetRules Target) : base(Target)
+	{
+		PublicDependencyModuleNames.AddRange(
+			new string[] {
+				"Core",
+				"CoreUObject",
+			}
+		);
+	}
+}
diff --git a/Engine/Source/Runtime/FieldNotification/Private/FieldNotificationDelegate.cpp b/Engine/Source/Runtime/FieldNotification/Private/FieldNotificationDelegate.cpp
new file mode 100644
--- /dev/null
+++ b/Engine/Source/Runtime/FieldNotification/Private/FieldNotificationDelegate.cpp
@@ -0,0 +1,524 @@
+// Copyright Epic Games, Inc. All Rights Reserved.
+
+#include "FieldNotificationDelegate.h"
+
+#include "UObject/Object.h"
+
+
+namespace UE::FieldNotification
+{
+
+FDelegateHandle FFieldMulticastDelegate::Add(const UObject* InObject, FFieldId InFieldId, FDelegate InNewDelegate)
+{
+	if (!InNewDelegate.IsBound())
+	{
+		return FDelegateHandle();
+	}
+
+	FDelegateHandle Result = InNewDelegate.GetHandle();
+	FInvocationKey InvocationKey{ InObject, InFieldId, NAME_None };
+	if (DelegateLockCount > 0)
+	{
+		const int32 Index = Delegates.Emplace(FInvocationElement{ InvocationKey, MoveTemp(InNewDelegate) });
+		AddedEmplaceAt = FMath::Min(AddedEmplaceAt, (uint16)Index);
+	}
+	else
+	{
+		const int32 FoundIndex = UpperBound(InvocationKey);
+		Delegates.EmplaceAt(FoundIndex, FInvocationElement{ InvocationKey, MoveTemp(InNewDelegate) });
+	}
+	return Result;
+}
+
+
+FDelegateHandle FFieldMulticastDelegate::Add(const UObject* InObject, FFieldId InFieldId, const FDynamicDelegate& InDynamicDelegate)
+{
+	if (!InDynamicDelegate.IsBound())
+	{
+		return FDelegateHandle();
+	}
+
+	FDelegate NewDelegate = FDelegate::CreateUFunction(const_cast<UObject*>(InDynamicDelegate.GetUObject()), InDynamicDelegate.GetFunctionName());
+	FDelegateHandle Result = NewDelegate.GetHandle();
+	FInvocationKey InvocationKey{ InObject, InFieldId, InDynamicDelegate.GetFunctionName()};
+	if (DelegateLockCount > 0)
+	{
+		const int32 Index = Delegates.Emplace(FInvocationElement{ InvocationKey, MoveTemp(NewDelegate) });
+		AddedEmplaceAt = FMath::Min(AddedEmplaceAt, (uint16)Index);
+	}
+	else
+	{
+		const int32 FoundIndex = UpperBound(InvocationKey);
+		Delegates.EmplaceAt(FoundIndex, FInvocationElement{ InvocationKey, MoveTemp(NewDelegate) });
+	}
+	return Result;
+}
+
+
+void FFieldMulticastDelegate::RemoveElement(FInvocationElement& Element, int32 Index, FRemoveResult& Result)
+{
+	Result.Object = Element.Key.Object.Get();
+	Result.FieldId = Element.Key.Id;
+	Result.bRemoved = Element.Delegate.IsBound();
+	if (DelegateLockCount > 0)
+	{
+		Element.Delegate.Unbind();
+		++CompactionCount;
+	}
+	else
+	{
+		Delegates.RemoveAt(Index);
+	}
+}
+
+
+void FFieldMulticastDelegate::CompleteRemove(FRemoveResult& Result) const
+{
+	if (Result.FieldId.IsValid())
+	{
+		// Search in the normal sorted list
+		const int32 FoundIndex = UpperBound(Result.FieldId) - 1;
+		if (Delegates.IsValidIndex(FoundIndex))
+		{
+			for (int32 Index = FoundIndex; Index >= 0; --Index)
+			{
+				const FInvocationElement& Element = Delegates[Index];
+				if (Element.Key.Id != Result.FieldId)
+				{
+					break;
+				}
+				else if (Element.Key.Object == Result.Object && Element.Delegate.IsBound())
+				{
+					Result.bHasOtherBoundDelegates = true;
+					break;
+				}
+			}
+		}
+
+		if (!Result.bHasOtherBoundDelegates)
+		{
+			for (int32 Index = AddedEmplaceAt; Index < Delegates.Num(); ++Index)
+			{
+				const FInvocationElement& Element = Delegates[Index];
+				if (Element.Key.Id == Result.FieldId && Element.Key.Object == Result.Object && Element.Delegate.IsBound())
+				{
+					Result.bHasOtherBoundDelegates = true;
+					break;
+				}
+			}
+		}
+	}
+}
+
+
+FFieldMulticastDelegate::FRemoveResult FFieldMulticastDelegate::Remove(FDelegateHandle InDelegate)
+{
+	FRemoveResult Result;
+
+	if (!InDelegate.IsValid())
+	{
+		return Result;
+	}
+
+	for (int32 Index = Delegates.Num() - 1; Index >= 0; --Index)
+	{
+		FInvocationElement& Element = Delegates[Index];
+		if (Element.Delegate.GetHandle() == InDelegate)
+		{
+			RemoveElement(Element, Index, Result);
+			break;
+		}
+	}
+
+	CompleteRemove(Result);
+
+	return Result;
+}
+
+
+FFieldMulticastDelegate::FRemoveResult FFieldMulticastDelegate::Remove(const FDynamicDelegate& InDynamicDelegate)
+{
+	FRemoveResult Result;
+
+	if (!InDynamicDelegate.IsBound())
+	{
+		// The UObject may be GCed. We would like to compact the list but there is not way to know which one (since HasSameObject will return false).
+		//Result may indicate that it was not remove even if it will be removed later when we finally compact.
+		return Result;
+	}
+
+	for (int32 Index = Delegates.Num() - 1; Index >= 0; --Index)
+	{
+		FInvocationElement& Element = Delegates[Index];
+		if (Element.Key.DynamicName == InDynamicDelegate.GetFunctionName())
+		{
+			if (const IDelegateInstance* DelegateInstance = GetDelegateInstance(Element.Delegate))
+			{
+				if (DelegateInstance->HasSameObject(InDynamicDelegate.GetUObject()))
+				{
+					RemoveElement(Element, Index, Result);
+					break;
+				}
+			}
+		}
+	}
+
+	CompleteRemove(Result);
+
+	return Result;
+}
+
+
+template<typename TRemoveOrUnbind>
+void FFieldMulticastDelegate::RemoveFrom(TRemoveOrUnbind& RemoveOrUnbind, FFieldId InFieldId, bool& bFieldPresent)
+{
+	// Search in the normal sorted list
+	const int32 FoundIndex = UpperBound(InFieldId) - 1;
+	if (Delegates.IsValidIndex(FoundIndex))
+	{
+		for (int32 Index = FoundIndex; Index >= 0; --Index)
+		{
+			FInvocationElement& Element = Delegates[Index];
+			if (Element.Key.Id != InFieldId)
+			{
+				break;
+			}
+			if (!RemoveOrUnbind(Element, Index))
+			{
+				break;
+			}
+		}
+	}
+
+	// The item may be at the end of the array, if added while broadcasting
+	if (!bFieldPresent)
+	{
+		for (int32 Index = Delegates.Num() - 1; Index >= AddedEmplaceAt; --Index)
+		{
+			FInvocationElement& Element = Delegates[Index];
+			if (Element.Key.Id == InFieldId)
+			{
+				if (!RemoveOrUnbind(Element, Index))
+				{
+					break;
+				}
+			}
+		}
+	}
+}
+
+
+FFieldMulticastDelegate::FRemoveFromResult FFieldMulticastDelegate::RemoveFrom(const UObject* InObject, FFieldId InFieldId, FDelegateHandle InDelegate)
+{
+	bool bRemoved = false;
+	bool bFieldPresent = false;
+
+	if (!InDelegate.IsValid())
+	{
+		return FRemoveFromResult{ bRemoved, bFieldPresent };
+	}
+
+	FFieldMulticastDelegate* Self = this;
+	auto RemoveOrUnbind = [Self, &InDelegate, &bRemoved, &bFieldPresent, InObject](FInvocationElement& Element, int32 Index)
+	{
+		if (Element.Delegate.GetHandle() == InDelegate)
+		{
+			if (Self->DelegateLockCount > 0)
+			{
+				Element.Delegate.Unbind();
+				++Self->CompactionCount;
+			}
+			else
+			{
+				Self->Delegates.RemoveAt(Index);
+			}
+			bRemoved = true;
+			return !bFieldPresent;
+		}
+		if (Element.Key.Object == InObject && Element.Delegate.IsBound())
+		{
+			bFieldPresent = true;
+			return !bRemoved;
+		}
+		return true;
+	};
+
+	RemoveFrom(RemoveOrUnbind, InFieldId, bFieldPresent);
+
+	return FRemoveFromResult{ bRemoved, bFieldPresent };
+}
+
+
+FFieldMulticastDelegate::FRemoveFromResult FFieldMulticastDelegate::RemoveFrom(const UObject* InObject, FFieldId InFieldId, const FDynamicDelegate& InDynamicDelegate)
+{
+	bool bRemoved = false;
+	bool bFieldPresent = false;
+
+	if (!InDynamicDelegate.IsBound())
+	{
+		return FRemoveFromResult{ bRemoved, bFieldPresent };
+	}
+
+	FFieldMulticastDelegate* Self = this;
+	auto RemoveOrUnbind = [Self, &InDynamicDelegate, &bRemoved, &bFieldPresent, InObject](FInvocationElement& Element, int32 Index)
+	{
+		if (Element.Key.DynamicName == InDynamicDelegate.GetFunctionName())
+		{
+			if (const IDelegateInstance* DelegateInstance = Self->GetDelegateInstance(Element.Delegate))
+			{
+				if (DelegateInstance->HasSameObject(InDynamicDelegate.GetUObject()))
+				{
+					if (Self->DelegateLockCount > 0)
+					{
+						Element.Delegate.Unbind();
+						++Self->CompactionCount;
+					}
+					else
+					{
+						Self->Delegates.RemoveAt(Index);
+					}
+					bRemoved = true;
+					return !bFieldPresent;
+				}
+			}
+		}
+		if (Element.Key.Object == InObject && Element.Delegate.IsBound())
+		{
+			bFieldPresent = true;
+			return !bRemoved;
+		}
+		return true;
+	};
+
+	RemoveFrom(RemoveOrUnbind, InFieldId, bFieldPresent);
+
+	return FRemoveFromResult{ bRemoved, bFieldPresent };
+}
+
+
+FFieldMulticastDelegate::FRemoveAllResult FFieldMulticastDelegate::RemoveAll(const UObject* InObject, const void* InUserObject)
+{
+	FRemoveAllResult Result;
+
+	auto SetHasFields = [&Result, InObject](FInvocationElement& Element)
+	{
+		Result.HasFields.PadToNum(Element.Key.Id.GetIndex() + 1, false);
+		Result.HasFields[Element.Key.Id.GetIndex()] = true;
+	};
+
+	if (DelegateLockCount > 0)
+	{
+		for (FInvocationElement& Element : Delegates)
+		{
+			if (const IDelegateInstance* DelegateInstance = GetDelegateInstance(Element.Delegate))
+			{
+				if (Element.Key.Object == InObject)
+				{
+					if (DelegateInstance->HasSameObject(InUserObject))
+					{
+						Element.Delegate.Unbind();
+						++CompactionCount;
+						++Result.RemoveCount;
+					}
+					else if (!DelegateInstance->IsCompactable())
+					{
+						SetHasFields(Element);
+					}
+				}
+			}
+		}
+	}
+	else
+	{
+		for (int32 Index = Delegates.Num() - 1; Index >= 0; --Index)
+		{
+			FInvocationElement& Element = Delegates[Index];
+			const IDelegateInstance* DelegateInstance = GetDelegateInstance(Element.Delegate);
+			if (DelegateInstance == nullptr || DelegateInstance->IsCompactable())
+			{
+				Delegates.RemoveAt(Index);
+			}
+			else if (Element.Key.Object == InObject)
+			{
+				if (DelegateInstance->HasSameObject(InUserObject))
+				{
+					Delegates.RemoveAt(Index);
+					++Result.RemoveCount;
+				}
+				else
+				{
+					SetHasFields(Element);
+				}
+			}
+		}
+		CompactionCount = 0;
+	}
+
+	return Result;
+}
+
+
+FFieldMulticastDelegate::FRemoveAllResult FFieldMulticastDelegate::RemoveAll(const UObject* InObject, FFieldId InFieldId, const void* InUserObject)
+{
+	FRemoveAllResult Result;
+
+	auto SetHasFields = [&Result, InObject](FInvocationElement& Element)
+	{
+		Result.HasFields.PadToNum(Element.Key.Id.GetIndex() + 1, false);
+		Result.HasFields[Element.Key.Id.GetIndex()] = true;
+	};
+
+	if (DelegateLockCount > 0)
+	{
+		for (FInvocationElement& Element : Delegates)
+		{
+			if (const IDelegateInstance* DelegateInstance = GetDelegateInstance(Element.Delegate))
+			{
+				if (Element.Key.Id == InFieldId && Element.Key.Object == InObject)
+				{
+					if (DelegateInstance->HasSameObject(InUserObject))
+					{
+						Element.Delegate.Unbind();
+						++CompactionCount;
+						++Result.RemoveCount;
+					}
+					else if (!DelegateInstance->IsCompactable())
+					{
+						SetHasFields(Element);
+					}
+				}
+			}
+		}
+	}
+	else
+	{
+		for (int32 Index = Delegates.Num() - 1; Index >= 0; --Index)
+		{
+			FInvocationElement& Element = Delegates[Index];
+			const IDelegateInstance* DelegateInstance = GetDelegateInstance(Element.Delegate);
+			if (DelegateInstance == nullptr || DelegateInstance->IsCompactable())
+			{
+				Delegates.RemoveAt(Index);
+			}
+			else if (Element.Key.Id == InFieldId && Element.Key.Object == InObject)
+			{
+				if (DelegateInstance->HasSameObject(InUserObject))
+				{
+					Delegates.RemoveAt(Index);
+					++Result.RemoveCount;
+				}
+				else
+				{
+					SetHasFields(Element);
+				}
+			}
+		}
+		CompactionCount = 0;
+	}
+
+	return Result;
+}
+
+
+void FFieldMulticastDelegate::Broadcast(UObject* InObject, UE::FieldNotification::FFieldId InFieldId)
+{
+	++DelegateLockCount;
+
+	// Search in the normal sorted list
+	const int32 FoundIndex = UpperBound(InFieldId) - 1;
+	if (Delegates.IsValidIndex(FoundIndex))
+	{
+		for (int32 Index = FoundIndex; Index >= 0; --Index)
+		{
+			FInvocationElement& Element = Delegates[Index];
+			if (Element.Key.Id != InFieldId)
+			{
+				break;
+			}
+			if (Element.Key.Object == InObject)
+			{
+				Element.Delegate.ExecuteIfBound(InObject, InFieldId);
+			}
+		}
+	}
+
+	// The item may be at the end of the array, if added while broadcasting
+	for (int32 Index = Delegates.Num() - 1; Index >= AddedEmplaceAt; --Index)
+	{
+		FInvocationElement& Element = Delegates[Index];
+		if (Element.Key.Id == InFieldId && Element.Key.Object == InObject)
+		{
+			Element.Delegate.ExecuteIfBound(InObject, InFieldId);
+		}
+	}
+
+	--DelegateLockCount;
+	ExecuteLockOperations();
+}
+
+
+void FFieldMulticastDelegate::Reset()
+{
+	if (DelegateLockCount > 0)
+	{
+		for (FInvocationElement& Element : Delegates)
+		{
+			CompactionCount += Element.Delegate.IsBound() ? 1 : 0;
+			Element.Delegate.Unbind();
+		}
+	}
+	else
+	{
+		Delegates.Reset();
+		DelegateLockCount = 0;
+		CompactionCount = 0;
+		AddedEmplaceAt = std::numeric_limits<uint16>::max();
+	}
+}
+
+
+void FFieldMulticastDelegate::ExecuteLockOperations()
+{
+	if (DelegateLockCount <= 0)
+	{
+		// Remove items that got removed while broadcasting
+		if (CompactionCount > 2)
+		{
+			for (int32 Index = Delegates.Num() - 1; Index >= 0; --Index)
+			{
+				FInvocationElement& Element = Delegates[Index];
+				const IDelegateInstance* DelegateInstance = GetDelegateInstance(Element.Delegate);
+				if (DelegateInstance == nullptr || DelegateInstance->IsCompactable())
+				{
+					Delegates.RemoveAt(Index);
+				}
+			}
+			CompactionCount = 0;
+		}
+
+		// Sort item that were added while broadcasting
+		for (; AddedEmplaceAt < Delegates.Num();)
+		{
+			if (Delegates[AddedEmplaceAt].Delegate.IsBound())
+			{
+				const int32 FoundIndex = UpperBound(Delegates[AddedEmplaceAt].Key);
+				if (FoundIndex != AddedEmplaceAt)
+				{
+					FInvocationElement Element = MoveTemp(Delegates[AddedEmplaceAt]);
+					Delegates.RemoveAtSwap(AddedEmplaceAt);
+					Delegates.EmplaceAt(FoundIndex, MoveTemp(Element));
+				}
+				++AddedEmplaceAt;
+			}
+			else
+			{
+				Delegates.RemoveAtSwap(AddedEmplaceAt);
+			}
+		}
+		AddedEmplaceAt = std::numeric_limits<uint16>::max();
+
+		DelegateLockCount = 0;
+	}
+}
+
+} //namespace
\ No newline at end of file
diff --git a/Engine/Source/Runtime/FieldNotification/Private/FieldNotificationModule.cpp b/Engine/Source/Runtime/FieldNotification/Private/FieldNotificationModule.cpp
new file mode 100644
--- /dev/null
+++ b/Engine/Source/Runtime/FieldNotification/Private/FieldNotificationModule.cpp
@@ -0,0 +1,5 @@
+// Copyright Epic Games, Inc. All Rights Reserved.
+
+#include "Modules/ModuleManager.h"
+
+IMPLEMENT_MODULE( FDefaultModuleImpl, FieldNotification );
diff --git a/Engine/Source/Runtime/FieldNotification/Private/Tests/FieldNotificationMulticastDelegateTest.cpp b/Engine/Source/Runtime/FieldNotification/Private/Tests/FieldNotificationMulticastDelegateTest.cpp
new file mode 100644
--- /dev/null
+++ b/Engine/Source/Runtime/FieldNotification/Private/Tests/FieldNotificationMulticastDelegateTest.cpp
@@ -0,0 +1,155 @@
+// Copyright Epic Games, Inc. All Rights Reserved.
+
+#include "Misc/AutomationTest.h"
+#include "FieldNotificationDelegate.h"
+
+
+#if WITH_AUTOMATION_WORKER
+
+#define LOCTEXT_NAMESPACE "FieldNotificationMulticastDelegateTest"
+
+IMPLEMENT_SIMPLE_AUTOMATION_TEST(FFieldNotificationMulticastDelegateTest, "System.UMG.FieldNotificationMulticastDelegateTest", EAutomationTestFlags::ApplicationContextMask | EAutomationTestFlags::EngineFilter);
+
+
+bool FFieldNotificationMulticastDelegateTest::RunTest(const FString& Parameters)
+{
+	struct FLocal
+	{
+		FLocal(FFieldNotificationMulticastDelegateTest* Value)
+			: Self(Value)
+		{}
+		FFieldNotificationMulticastDelegateTest* Self = nullptr;
+		UE::FieldNotification::FFieldMulticastDelegate Delegate;
+		int32 FooACount = 0;
+		int32 FooBCount = 0;
+		int32 FooCCount = 0;
+		UE::FieldNotification::FFieldId ExpectFieldA;
+		UE::FieldNotification::FFieldId ExpectFieldB;
+		UE::FieldNotification::FFieldId ExpectFieldC;
+
+		void Reset()
+		{
+			FooACount = 0;
+			FooBCount = 0;
+			FooCCount = 0;
+			ExpectFieldA = UE::FieldNotification::FFieldId();
+			ExpectFieldB = UE::FieldNotification::FFieldId();
+			ExpectFieldC = UE::FieldNotification::FFieldId();
+		}
+		void FooA(UObject* InObject, UE::FieldNotification::FFieldId Id)
+		{
+			++FooACount;
+			if (ExpectFieldA.IsValid() && Id != ExpectFieldA)
+			{
+				Self->AddError(TEXT("The expected field A is not valid."));
+			}
+		}
+		void FooB(UObject* InObject, UE::FieldNotification::FFieldId Id)
+		{
+			++FooBCount;
+			if (ExpectFieldB.IsValid() && Id != ExpectFieldB)
+			{
+				Self->AddError(TEXT("The expected field B is not valid."));
+			}
+		}
+		void FooC(UObject* InObject, UE::FieldNotification::FFieldId Id)
+		{
+			++FooCCount;
+			if (ExpectFieldC.IsValid() && Id != ExpectFieldC)
+			{
+				Self->AddError(TEXT("The expected field C is not valid."));
+			}
+		}
+	};
+
+	UE::FieldNotification::FFieldId FieldA{ "FieldA", 0 };
+	UE::FieldNotification::FFieldId FieldB{ "FieldB", 1 };
+	UE::FieldNotification::FFieldId FieldC{ "FieldC", 2 };
+	UE::FieldNotification::FFieldId FieldD{ "FieldD", 3 };
+
+	FLocal Local{ this };
+	auto RunTest1 = [&Local, &FieldA, &FieldB, &FieldC, this]()
+	{
+		{
+			Local.Reset();
+			FDelegateHandle Handle1 = Local.Delegate.Add(UObject::StaticClass(), FieldA, UE::FieldNotification::FFieldMulticastDelegate::FDelegate::CreateRaw(&Local, &FLocal::FooA));
+			Local.ExpectFieldA = FieldA;
+			Local.Delegate.Broadcast(UObject::StaticClass(), FieldA);
+			AddErrorIfFalse(Local.FooACount == 1, TEXT("Bad amount of FooA"));
+			AddErrorIfFalse(Local.FooBCount == 0, TEXT("Bad amount of FooB"));
+
+			Local.Reset();
+			FDelegateHandle Handle2 = Local.Delegate.Add(UObject::StaticClass(), FieldA, UE::FieldNotification::FFieldMulticastDelegate::FDelegate::CreateRaw(&Local, &FLocal::FooA));
+			Local.ExpectFieldA = FieldA;
+			Local.Delegate.Broadcast(UObject::StaticClass(), FieldA);
+			AddErrorIfFalse(Local.FooACount == 2, TEXT("Bad amount of FooA"));
+			AddErrorIfFalse(Local.FooBCount == 0, TEXT("Bad amount of FooB"));
+
+			UE::FieldNotification::FFieldMulticastDelegate::FRemoveFromResult Result1 = Local.Delegate.RemoveFrom(UObject::StaticClass(), FieldA, Handle1);
+			AddErrorIfFalse(Result1.bRemoved, TEXT("Bad Remove of Handle1"));
+			AddErrorIfFalse(Result1.bHasOtherBoundDelegates, TEXT("Remove of Handle1 should have other bound delegates."));
+			UE::FieldNotification::FFieldMulticastDelegate::FRemoveFromResult Result2 = Local.Delegate.RemoveFrom(UObject::StaticClass(), FieldA, Handle1);
+			AddErrorIfFalse(!Result2.bRemoved, TEXT("Bad Remove of Handle1"));
+			AddErrorIfFalse(Result2.bHasOtherBoundDelegates, TEXT("Remove of Handle1 should not have other bound delegates."));
+			UE::FieldNotification::FFieldMulticastDelegate::FRemoveFromResult Result3 = Local.Delegate.RemoveFrom(UObject::StaticClass(), FieldA, Handle2);
+			AddErrorIfFalse(Result3.bRemoved, TEXT("Bad Remove of Handle2"));
+			AddErrorIfFalse(!Result3.bHasOtherBoundDelegates, TEXT("Remove of Handle2 should not have other bound delegates."));
+		}
+
+		{
+			Local.Reset();
+			FDelegateHandle Handle1 = Local.Delegate.Add(UObject::StaticClass(), FieldA, UE::FieldNotification::FFieldMulticastDelegate::FDelegate::CreateRaw(&Local, &FLocal::FooA));
+			FDelegateHandle Handle2 = Local.Delegate.Add(UClass::StaticClass(), FieldA, UE::FieldNotification::FFieldMulticastDelegate::FDelegate::CreateRaw(&Local, &FLocal::FooA));
+			Local.Delegate.Add(UClass::StaticClass(), FieldB, UE::FieldNotification::FFieldMulticastDelegate::FDelegate::CreateRaw(&Local, &FLocal::FooB));
+			Local.Delegate.Add(UClass::StaticClass(), FieldC, UE::FieldNotification::FFieldMulticastDelegate::FDelegate::CreateRaw(&Local, &FLocal::FooC));
+			Local.ExpectFieldA = FieldA;
+			Local.ExpectFieldB = FieldB;
+			Local.ExpectFieldC = FieldC;
+			Local.Delegate.Broadcast(UObject::StaticClass(), FieldA);
+			AddErrorIfFalse(Local.FooACount == 1, TEXT("Bad amount of FooA"));
+			AddErrorIfFalse(Local.FooBCount == 0, TEXT("Bad amount of FooB"));
+			AddErrorIfFalse(Local.FooCCount == 0, TEXT("Bad amount of FooC"));
+			Local.Delegate.Broadcast(UClass::StaticClass(), FieldA);
+			AddErrorIfFalse(Local.FooACount == 2, TEXT("Bad amount of FooA"));
+			AddErrorIfFalse(Local.FooBCount == 0, TEXT("Bad amount of FooB"));
+			AddErrorIfFalse(Local.FooCCount == 0, TEXT("Bad amount of FooC"));
+			Local.Delegate.Broadcast(UClass::StaticClass(), FieldB);
+			AddErrorIfFalse(Local.FooACount == 2, TEXT("Bad amount of FooA"));
+			AddErrorIfFalse(Local.FooBCount == 1, TEXT("Bad amount of FooB"));
+			AddErrorIfFalse(Local.FooCCount == 0, TEXT("Bad amount of FooC"));
+			Local.Delegate.Broadcast(UClass::StaticClass(), FieldC);
+			AddErrorIfFalse(Local.FooACount == 2, TEXT("Bad amount of FooA"));
+			AddErrorIfFalse(Local.FooBCount == 1, TEXT("Bad amount of FooB"));
+			AddErrorIfFalse(Local.FooCCount == 1, TEXT("Bad amount of FooC"));
+
+			UE::FieldNotification::FFieldMulticastDelegate::FRemoveResult Result1 = Local.Delegate.Remove(Handle2);
+			AddErrorIfFalse(Result1.Object == UClass::StaticClass(), TEXT("Failed Remove Handle2, bad Object"));
+			AddErrorIfFalse(Result1.FieldId == FieldA, TEXT("Failed Remove Handle2, bad FieldId"));
+			AddErrorIfFalse(Result1.bRemoved, TEXT("Failed Remove Handle2, bad bRemoved"));
+			AddErrorIfFalse(!Result1.bHasOtherBoundDelegates, TEXT("Failed Remove Handle2, bad bHasOtherBoundDelegates"));
+
+			UE::FieldNotification::FFieldMulticastDelegate::FRemoveAllResult Result2 = Local.Delegate.RemoveAll(UClass::StaticClass(), &Local);
+			AddErrorIfFalse(Result2.RemoveCount == 2, TEXT("Failed RemoveAll, bad count"));
+			AddErrorIfFalse(Result2.HasFields.CountSetBits() == 0, TEXT("Failed RemoveAll, bad HasFields"));
+
+			UE::FieldNotification::FFieldMulticastDelegate::FRemoveFromResult Result3 = Local.Delegate.RemoveFrom(UObject::StaticClass(), FieldA, Handle1);
+			AddErrorIfFalse(Result3.bRemoved, TEXT("Bad Remove of FieldA"));
+			AddErrorIfFalse(!Result3.bHasOtherBoundDelegates, TEXT("Remove of Handle1 should have other bound delegates."));
+		}
+	};
+	RunTest1();
+
+	Local.Delegate.Add(
+			UObject::StaticClass(),
+			FieldD,
+			UE::FieldNotification::FFieldMulticastDelegate::FDelegate::CreateLambda([&RunTest1](UObject* InObject, UE::FieldNotification::FFieldId Id)
+			{
+				RunTest1();
+			}));
+	Local.Delegate.Broadcast(UObject::StaticClass(), FieldD);
+
+	return true;
+}
+
+#undef LOCTEXT_NAMESPACE 
+#endif //WITH_AUTOMATION_WORKER
\ No newline at end of file
diff --git a/Engine/Source/Runtime/FieldNotification/Public/FieldNotificationDeclaration.h b/Engine/Source/Runtime/FieldNotification/Public/FieldNotificationDeclaration.h
new file mode 100644
--- /dev/null
+++ b/Engine/Source/Runtime/FieldNotification/Public/FieldNotificationDeclaration.h
@@ -0,0 +1,244 @@
+// Copyright Epic Games, Inc. All Rights Reserved.
+
+#pragma once
+
+#include "FieldNotificationId.h"
+#include "IFieldNotificationClassDescriptor.h"
+
+
+/*
+ *	struct FFieldNotificationClassDescriptor : Super::FFieldNotificationClassDescriptor
+ *	{
+ *	private:
+ *		using SuperDescriptor = Super::FFieldNotificationClassDescriptor;
+ *		static const ::UE::FieldNotification::FFieldId* AllFields[];
+ *		friend ThisClass;
+ *	public:
+ *		static const ::UE::FieldNotification::FFieldId FieldA;
+ *		static const ::UE::FieldNotification::FFieldId FieldB;
+ *		enum
+ *		{
+ *			IndexOf_FieldA = Super::FFieldNotificationClassDescriptor::Max_IndexOf_ + 0,
+ *			IndexOf_FieldB,
+ *			Max_IndexOf_,
+ *		};
+ * 	virtual void ForEachField(const UClass* Class, TFunctionRef<bool(::UE::FieldNotification::FFieldId FielId)> Callback) const override;
+ *	};
+ * virtual const ::UE::FieldNotification::IClassDescriptor& GetFieldNotificationDescriptor() const override;
+ */
+
+#define UE_FIELD_NOTIFICATION_OVERRIDE_GET_FIELD_NOTIFICATION_DESCRIPTOR() \
+	virtual const ::UE::FieldNotification::IClassDescriptor& GetFieldNotificationDescriptor() const \
+	{ \
+ 		static FFieldNotificationClassDescriptor Instance; \
+ 		return Instance; \
+ 	}
+
+#define UE_FIELD_NOTIFICATION_DECLARE_CLASS_DESCRIPTOR_BEGIN(API_STRING) \
+	struct API_STRING FFieldNotificationClassDescriptor : public Super::FFieldNotificationClassDescriptor \
+	{ \
+	private: \
+		using SuperDescriptor = Super::FFieldNotificationClassDescriptor; \
+		static const ::UE::FieldNotification::FFieldId* AllFields[]; \
+		friend ThisClass; \
+	public:
+
+
+#define UE_FIELD_NOTIFICATION_DECLARE_CLASS_DESCRIPTOR_END() \
+		virtual void ForEachField(const UClass* Class, TFunctionRef<bool(::UE::FieldNotification::FFieldId FielId)> Callback) const override \
+		{ \
+			for (int32 Index = 0; Index < Max_IndexOf_-SuperDescriptor::Max_IndexOf_; ++Index) \
+			{ \
+				if (!Callback(*AllFields[Index])) \
+				{ \
+					return; \
+				} \
+			} \
+			SuperDescriptor::ForEachField(Class, Callback); \
+		} \
+	}; \
+	UE_FIELD_NOTIFICATION_OVERRIDE_GET_FIELD_NOTIFICATION_DESCRIPTOR()
+
+
+#define UE_FIELD_NOTIFICATION_DECLARE_CLASS_DESCRIPTOR_BASE_BEGIN(API_STRING) \
+	struct API_STRING FFieldNotificationClassDescriptor : public ::UE::FieldNotification::IClassDescriptor \
+	{ \
+	private: \
+		using SuperDescriptor = ::UE::FieldNotification::IClassDescriptor; \
+		static const ::UE::FieldNotification::FFieldId* AllFields[]; \
+		friend ThisClass; \
+	public:
+
+
+#define UE_FIELD_NOTIFICATION_DECLARE_CLASS_DESCRIPTOR_BASE_END() \
+		virtual void ForEachField(const UClass* Class, TFunctionRef<bool(::UE::FieldNotification::FFieldId FielId)> Callback) const override; \
+	}; \
+	UE_FIELD_NOTIFICATION_OVERRIDE_GET_FIELD_NOTIFICATION_DESCRIPTOR()
+
+
+#define UE_FIELD_NOTIFICATION_DECLARE_ENUM_FIELD_BEGIN(Name) \
+		enum EField \
+		{ \
+			IndexOf_##Name = SuperDescriptor::Max_IndexOf_ + 0,
+
+
+#define UE_FIELD_NOTIFICATION_DECLARE_ENUM_FIELD_END() \
+			Max_IndexOf_, \
+		};
+
+
+#define UE_FIELD_NOTIFICATION_DECLARE_FIELD(Name) \
+	static const ::UE::FieldNotification::FFieldId Name;
+
+
+#define UE_FIELD_NOTIFICATION_DECLARE_ENUM_FIELD(Name) \
+			IndexOf_##Name,
+
+
+#define UE_FIELD_NOTIFICATION_IMPLEMENTATION_BEGIN(ClassName) \
+	const ::UE::FieldNotification::FFieldId* ClassName::FFieldNotificationClassDescriptor::AllFields[] = {
+
+
+#define UE_FIELD_NOTIFICATION_IMPLEMENTATION_END(ClassName) }
+
+//#if DO_CHECK
+//#define UE_FIELD_NOTIFICATION_IMPLEMENTATION_END(ClassName) \
+//	}; \
+//	const ::UE::FieldNotification::IClassDescriptor& ClassName::GetFieldNotificationDescriptor() const \
+//	{ \
+//		static_assert(UE_ARRAY_COUNT(ClassName::FFieldNotificationClassDescriptor::AllFields) == ClassName::FFieldNotificationClassDescriptor::Max_IndexOf_-ClassName::FFieldNotificationClassDescriptor::SuperDescriptor::Max_IndexOf_, "The descriptor for class " #ClassName " doesn't not implement the same number of field as declared in in the constructor."); \
+//		struct FLocal \
+//		{ \
+//			ClassName::FFieldNotificationClassDescriptor Instance; \
+//			FLocal() \
+//			{ \
+//				int32 NumberOfField = Instance.GetNumberOfField(); \
+//				for (int32 Index = 0; Index < NumberOfField; ++Index) \
+//				{ \
+//					::UE::FieldNotification::FFieldId Id = Instance.GetField(Index); \
+//					ensureAlwaysMsgf(Id.IsValid(), TEXT("The FFieldId '%d' for the class '%s' is invalid"), Index, *ClassName::StaticClass()->GetName()); \
+//					ensureAlwaysMsgf(Id.GetIndex() == Index, TEXT("The FFieldId at index '%d' doesn't match the id index '%d'."), Index, Id.GetIndex()); \
+//				} \
+//			} \
+//		}; \
+//		static FLocal Local; \
+//		return Local.Instance; \
+//	}
+//#else
+//#define UE_FIELD_NOTIFICATION_IMPLEMENTATION_END(ClassName) \
+//	}; \
+//	const ::UE::FieldNotification::IClassDescriptor& ClassName::GetFieldNotificationDescriptor() const \
+//	{ \
+//		static FFieldNotificationClassDescriptor Instance; \
+//		return Instance; \
+//	}
+//#endif
+
+
+#define UE_FIELD_NOTIFICATION_IMPLEMENT_FIELD(ClassName, MemberName) const ::UE::FieldNotification::FFieldId ClassName::FFieldNotificationClassDescriptor::MemberName( FName(TEXT(#MemberName)), ClassName::FFieldNotificationClassDescriptor::IndexOf_##MemberName );
+#define UE_FIELD_NOTIFICATION_IMPLEMENT_ENUM_FIELD(ClassName, MemberName) &ClassName::FFieldNotificationClassDescriptor::MemberName,
+
+
+#define UE_FIELD_NOTIFICATION_IMPLEMENT_CLASS_DESCRIPTOR_OneField(ClassName, Field1) \
+	UE_FIELD_NOTIFICATION_IMPLEMENT_FIELD(ClassName, Field1) \
+	UE_FIELD_NOTIFICATION_IMPLEMENTATION_BEGIN(ClassName) \
+	UE_FIELD_NOTIFICATION_IMPLEMENT_ENUM_FIELD(ClassName, Field1) \
+	UE_FIELD_NOTIFICATION_IMPLEMENTATION_END(ClassName)
+
+#define UE_FIELD_NOTIFICATION_IMPLEMENT_CLASS_DESCRIPTOR_TwoFields(ClassName, Field1, Field2) \
+	UE_FIELD_NOTIFICATION_IMPLEMENT_FIELD(ClassName, Field1) \
+	UE_FIELD_NOTIFICATION_IMPLEMENT_FIELD(ClassName, Field2) \
+	UE_FIELD_NOTIFICATION_IMPLEMENTATION_BEGIN(ClassName) \
+	UE_FIELD_NOTIFICATION_IMPLEMENT_ENUM_FIELD(ClassName, Field1) \
+	UE_FIELD_NOTIFICATION_IMPLEMENT_ENUM_FIELD(ClassName, Field2) \
+	UE_FIELD_NOTIFICATION_IMPLEMENTATION_END(ClassName)
+
+#define UE_FIELD_NOTIFICATION_IMPLEMENT_CLASS_DESCRIPTOR_ThreeFields(ClassName, Field1, Field2, Field3) \
+	UE_FIELD_NOTIFICATION_IMPLEMENT_FIELD(ClassName, Field1) \
+	UE_FIELD_NOTIFICATION_IMPLEMENT_FIELD(ClassName, Field2) \
+	UE_FIELD_NOTIFICATION_IMPLEMENT_FIELD(ClassName, Field3) \
+	UE_FIELD_NOTIFICATION_IMPLEMENTATION_BEGIN(ClassName) \
+	UE_FIELD_NOTIFICATION_IMPLEMENT_ENUM_FIELD(ClassName, Field1) \
+	UE_FIELD_NOTIFICATION_IMPLEMENT_ENUM_FIELD(ClassName, Field2) \
+	UE_FIELD_NOTIFICATION_IMPLEMENT_ENUM_FIELD(ClassName, Field3) \
+	UE_FIELD_NOTIFICATION_IMPLEMENTATION_END(ClassName)
+
+#define UE_FIELD_NOTIFICATION_IMPLEMENT_CLASS_DESCRIPTOR_FourFields(ClassName, Field1, Field2, Field3, Field4) \
+	UE_FIELD_NOTIFICATION_IMPLEMENT_FIELD(ClassName, Field1) \
+	UE_FIELD_NOTIFICATION_IMPLEMENT_FIELD(ClassName, Field2) \
+	UE_FIELD_NOTIFICATION_IMPLEMENT_FIELD(ClassName, Field3) \
+	UE_FIELD_NOTIFICATION_IMPLEMENT_FIELD(ClassName, Field4) \
+	UE_FIELD_NOTIFICATION_IMPLEMENTATION_BEGIN(ClassName) \
+	UE_FIELD_NOTIFICATION_IMPLEMENT_ENUM_FIELD(ClassName, Field1) \
+	UE_FIELD_NOTIFICATION_IMPLEMENT_ENUM_FIELD(ClassName, Field2) \
+	UE_FIELD_NOTIFICATION_IMPLEMENT_ENUM_FIELD(ClassName, Field3) \
+	UE_FIELD_NOTIFICATION_IMPLEMENT_ENUM_FIELD(ClassName, Field4) \
+	UE_FIELD_NOTIFICATION_IMPLEMENTATION_END(ClassName)
+
+#define UE_FIELD_NOTIFICATION_IMPLEMENT_CLASS_DESCRIPTOR_FiveFields(ClassName, Field1, Field2, Field3, Field4, Field5) \
+	UE_FIELD_NOTIFICATION_IMPLEMENT_FIELD(ClassName, Field1) \
+	UE_FIELD_NOTIFICATION_IMPLEMENT_FIELD(ClassName, Field2) \
+	UE_FIELD_NOTIFICATION_IMPLEMENT_FIELD(ClassName, Field3) \
+	UE_FIELD_NOTIFICATION_IMPLEMENT_FIELD(ClassName, Field4) \
+	UE_FIELD_NOTIFICATION_IMPLEMENT_FIELD(ClassName, Field5) \
+	UE_FIELD_NOTIFICATION_IMPLEMENTATION_BEGIN(ClassName) \
+	UE_FIELD_NOTIFICATION_IMPLEMENT_ENUM_FIELD(ClassName, Field1) \
+	UE_FIELD_NOTIFICATION_IMPLEMENT_ENUM_FIELD(ClassName, Field2) \
+	UE_FIELD_NOTIFICATION_IMPLEMENT_ENUM_FIELD(ClassName, Field3) \
+	UE_FIELD_NOTIFICATION_IMPLEMENT_ENUM_FIELD(ClassName, Field4) \
+	UE_FIELD_NOTIFICATION_IMPLEMENT_ENUM_FIELD(ClassName, Field5) \
+	UE_FIELD_NOTIFICATION_IMPLEMENTATION_END(ClassName)
+
+#define UE_FIELD_NOTIFICATION_IMPLEMENT_CLASS_DESCRIPTOR_SixFields(ClassName, Field1, Field2, Field3, Field4, Field5, Field6) \
+	UE_FIELD_NOTIFICATION_IMPLEMENT_FIELD(ClassName, Field1) \
+	UE_FIELD_NOTIFICATION_IMPLEMENT_FIELD(ClassName, Field2) \
+	UE_FIELD_NOTIFICATION_IMPLEMENT_FIELD(ClassName, Field3) \
+	UE_FIELD_NOTIFICATION_IMPLEMENT_FIELD(ClassName, Field4) \
+	UE_FIELD_NOTIFICATION_IMPLEMENT_FIELD(ClassName, Field5) \
+	UE_FIELD_NOTIFICATION_IMPLEMENT_FIELD(ClassName, Field6) \
+	UE_FIELD_NOTIFICATION_IMPLEMENTATION_BEGIN(ClassName) \
+	UE_FIELD_NOTIFICATION_IMPLEMENT_ENUM_FIELD(ClassName, Field1) \
+	UE_FIELD_NOTIFICATION_IMPLEMENT_ENUM_FIELD(ClassName, Field2) \
+	UE_FIELD_NOTIFICATION_IMPLEMENT_ENUM_FIELD(ClassName, Field3) \
+	UE_FIELD_NOTIFICATION_IMPLEMENT_ENUM_FIELD(ClassName, Field4) \
+	UE_FIELD_NOTIFICATION_IMPLEMENT_ENUM_FIELD(ClassName, Field5) \
+	UE_FIELD_NOTIFICATION_IMPLEMENT_ENUM_FIELD(ClassName, Field6) \
+	UE_FIELD_NOTIFICATION_IMPLEMENTATION_END(ClassName)
+
+#define UE_FIELD_NOTIFICATION_IMPLEMENT_CLASS_DESCRIPTOR_SevenFields(ClassName, Field1, Field2, Field3, Field4, Field5, Field6, Field7) \
+	UE_FIELD_NOTIFICATION_IMPLEMENT_FIELD(ClassName, Field1) \
+	UE_FIELD_NOTIFICATION_IMPLEMENT_FIELD(ClassName, Field2) \
+	UE_FIELD_NOTIFICATION_IMPLEMENT_FIELD(ClassName, Field3) \
+	UE_FIELD_NOTIFICATION_IMPLEMENT_FIELD(ClassName, Field4) \
+	UE_FIELD_NOTIFICATION_IMPLEMENT_FIELD(ClassName, Field5) \
+	UE_FIELD_NOTIFICATION_IMPLEMENT_FIELD(ClassName, Field6) \
+	UE_FIELD_NOTIFICATION_IMPLEMENT_FIELD(ClassName, Field7) \
+	UE_FIELD_NOTIFICATION_IMPLEMENTATION_BEGIN(ClassName) \
+	UE_FIELD_NOTIFICATION_IMPLEMENT_ENUM_FIELD(ClassName, Field1) \
+	UE_FIELD_NOTIFICATION_IMPLEMENT_ENUM_FIELD(ClassName, Field2) \
+	UE_FIELD_NOTIFICATION_IMPLEMENT_ENUM_FIELD(ClassName, Field3) \
+	UE_FIELD_NOTIFICATION_IMPLEMENT_ENUM_FIELD(ClassName, Field4) \
+	UE_FIELD_NOTIFICATION_IMPLEMENT_ENUM_FIELD(ClassName, Field5) \
+	UE_FIELD_NOTIFICATION_IMPLEMENT_ENUM_FIELD(ClassName, Field6) \
+	UE_FIELD_NOTIFICATION_IMPLEMENT_ENUM_FIELD(ClassName, Field7) \
+	UE_FIELD_NOTIFICATION_IMPLEMENTATION_END(ClassName)
+
+#define UE_FIELD_NOTIFICATION_IMPLEMENT_CLASS_DESCRIPTOR_EightFields(ClassName, Field1, Field2, Field3, Field4, Field5, Field6, Field7, Field8) \
+	UE_FIELD_NOTIFICATION_IMPLEMENT_FIELD(ClassName, Field1) \
+	UE_FIELD_NOTIFICATION_IMPLEMENT_FIELD(ClassName, Field2) \
+	UE_FIELD_NOTIFICATION_IMPLEMENT_FIELD(ClassName, Field3) \
+	UE_FIELD_NOTIFICATION_IMPLEMENT_FIELD(ClassName, Field4) \
+	UE_FIELD_NOTIFICATION_IMPLEMENT_FIELD(ClassName, Field5) \
+	UE_FIELD_NOTIFICATION_IMPLEMENT_FIELD(ClassName, Field6) \
+	UE_FIELD_NOTIFICATION_IMPLEMENT_FIELD(ClassName, Field7) \
+	UE_FIELD_NOTIFICATION_IMPLEMENT_FIELD(ClassName, Field8) \
+	UE_FIELD_NOTIFICATION_IMPLEMENTATION_BEGIN(ClassName) \
+	UE_FIELD_NOTIFICATION_IMPLEMENT_ENUM_FIELD(ClassName, Field1) \
+	UE_FIELD_NOTIFICATION_IMPLEMENT_ENUM_FIELD(ClassName, Field2) \
+	UE_FIELD_NOTIFICATION_IMPLEMENT_ENUM_FIELD(ClassName, Field3) \
+	UE_FIELD_NOTIFICATION_IMPLEMENT_ENUM_FIELD(ClassName, Field4) \
+	UE_FIELD_NOTIFICATION_IMPLEMENT_ENUM_FIELD(ClassName, Field5) \
+	UE_FIELD_NOTIFICATION_IMPLEMENT_ENUM_FIELD(ClassName, Field6) \
+	UE_FIELD_NOTIFICATION_IMPLEMENT_ENUM_FIELD(ClassName, Field7) \
+	UE_FIELD_NOTIFICATION_IMPLEMENT_ENUM_FIELD(ClassName, Field8) \
+	UE_FIELD_NOTIFICATION_IMPLEMENTATION_END(ClassName)
diff --git a/Engine/Source/Runtime/FieldNotification/Public/FieldNotificationDelegate.h b/Engine/Source/Runtime/FieldNotification/Public/FieldNotificationDelegate.h
new file mode 100644
--- /dev/null
+++ b/Engine/Source/Runtime/FieldNotification/Public/FieldNotificationDelegate.h
@@ -0,0 +1,155 @@
+// Copyright Epic Games, Inc. All Rights Reserved.
+
+#pragma once
+
+#include "Algo/BinarySearch.h"
+#include "Containers/Array.h"
+#include "Containers/BitArray.h"
+#include "CoreMinimal.h"
+#include "Delegates/Delegate.h"
+#include "FieldNotificationId.h"
+#include "INotifyFieldValueChanged.h"
+#include "Templates/IdentityFunctor.h"
+#include "Templates/Less.h"
+#include "Templates/UnrealTemplate.h"
+#include "UObject/NameTypes.h"
+#include "UObject/WeakObjectPtr.h"
+
+#include <limits>
+
+class UObject;
+
+namespace UE::FieldNotification
+{
+	
+class FIELDNOTIFICATION_API FFieldMulticastDelegate
+{
+public:
+	using FDelegate = INotifyFieldValueChanged::FFieldValueChangedDelegate;
+	using FDynamicDelegate = FFieldValueChangedDynamicDelegate;
+
+private:
+	struct FInstanceProtectedDelegate : public TDelegateBase<FDefaultDelegateUserPolicy>
+	{
+		using Super = TDelegateBase<FDefaultDelegateUserPolicy>;
+		using Super::GetDelegateInstanceProtected;
+	};
+
+	struct FInvocationKey
+	{
+		TWeakObjectPtr<const UObject> Object;
+		FFieldId Id;
+		FName DynamicName;
+		bool operator<(const FInvocationKey& Element) const
+		{
+			return Id.GetName().FastLess(Element.Id.GetName());
+		}
+	};
+
+	struct FInvocationElement
+	{
+		FInvocationKey Key;
+		FDelegate Delegate;
+		bool operator<(const FInvocationElement& InElement) const
+		{
+			return Key < InElement.Key;
+		}
+	};
+
+	friend bool operator<(const FInvocationElement& A, const FInvocationKey& B)
+	{
+		return A.Key < B;
+	}
+	friend bool operator<(const FInvocationKey& A, const FInvocationElement& B)
+	{
+		return A < B.Key;
+	}
+	friend bool operator<(const FInvocationElement& A, const FFieldId& B)
+	{
+		return A.Key.Id.GetName().FastLess(B.GetName());
+	}
+	friend bool operator<(const FFieldId& A, const FInvocationElement& B)
+	{
+		return A.GetName().FastLess(B.Key.Id.GetName());
+	}
+
+	using InvocationListType = TArray<FInvocationElement>;
+
+public:
+	FDelegateHandle Add(const UObject* InObject, FFieldId InFieldId, FDelegate InNewDelegate);
+	FDelegateHandle Add(const UObject* InObject, FFieldId InFieldId, const FDynamicDelegate& InDynamicDelegate);
+
+	struct FRemoveResult
+	{
+		bool bRemoved = false;
+		bool bHasOtherBoundDelegates = false;
+		const UObject* Object = nullptr;
+		FFieldId FieldId;
+	};
+	FRemoveResult Remove(FDelegateHandle InDelegate);
+	FRemoveResult Remove(const FDynamicDelegate& InDynamicDelegate);
+
+	struct FRemoveFromResult
+	{
+		bool bRemoved = false;
+		bool bHasOtherBoundDelegates = false;
+	};
+	FRemoveFromResult RemoveFrom(const UObject* InObject, FFieldId InFieldId, FDelegateHandle InDelegate);
+	FRemoveFromResult RemoveFrom(const UObject* InObject, FFieldId InFieldId, const FDynamicDelegate& InDynamicDelegate);
+
+	struct FRemoveAllResult
+	{
+		int32 RemoveCount = 0;
+		TBitArray<> HasFields;
+	};
+	FRemoveAllResult RemoveAll(const UObject* InObject, const void* InUserObject);
+	FRemoveAllResult RemoveAll(const UObject* InObject, FFieldId InFieldId, const void* InUserObject);
+
+	void Broadcast(UObject* InObject, FFieldId InFieldId);
+
+	void Reset();
+
+private:
+	int32 LowerBound(FFieldId InFieldId) const
+	{
+		int32 Num = AddedEmplaceAt > GetNum(Delegates) ? GetNum(Delegates) : AddedEmplaceAt;
+		return AlgoImpl::LowerBoundInternal(GetData(Delegates), Num, InFieldId, FIdentityFunctor(), TLess<>());
+	}
+
+	int32 LowerBound(const FInvocationKey& InKey) const
+	{
+		int32 Num = AddedEmplaceAt > GetNum(Delegates) ? GetNum(Delegates) : AddedEmplaceAt;
+		return AlgoImpl::LowerBoundInternal(GetData(Delegates), Num, InKey, FIdentityFunctor(), TLess<>());
+	}
+
+	int32 UpperBound(FFieldId InFieldId) const
+	{
+		int32 Num = AddedEmplaceAt > GetNum(Delegates) ? GetNum(Delegates) : AddedEmplaceAt;
+		return AlgoImpl::UpperBoundInternal(GetData(Delegates), Num, InFieldId, FIdentityFunctor(), TLess<>());
+	}
+
+	int32 UpperBound(const FInvocationKey& InKey) const
+	{
+		int32 Num = AddedEmplaceAt > GetNum(Delegates) ? GetNum(Delegates) : AddedEmplaceAt;
+		return AlgoImpl::UpperBoundInternal(GetData(Delegates), Num, InKey, FIdentityFunctor(), TLess<>());
+	}
+
+	const IDelegateInstance* GetDelegateInstance(FDelegate& InDelegate) const
+	{
+		return static_cast<FInstanceProtectedDelegate&>(static_cast<TDelegateBase<FDefaultDelegateUserPolicy>&>(InDelegate)).GetDelegateInstanceProtected();
+	}
+
+	void ExecuteLockOperations();
+	void RemoveElement(FInvocationElement& Element, int32 Index, FRemoveResult& Result);
+	void CompleteRemove(FRemoveResult& Result) const;
+	template<typename TRemoveOrUnbind>
+	void RemoveFrom(TRemoveOrUnbind& RemoveOrUnbind, FFieldId FieldId, bool& bFieldPresent);
+
+private:
+	InvocationListType Delegates;
+	int16 DelegateLockCount = 0;
+	int16 CompactionCount = 0;
+	uint16 AddedEmplaceAt = std::numeric_limits<uint16>::max();
+};
+
+} //namespace
\ No newline at end of file
diff --git a/Engine/Source/Runtime/FieldNotification/Public/FieldNotificationHelpers.h b/Engine/Source/Runtime/FieldNotification/Public/FieldNotificationHelpers.h
new file mode 100644
--- /dev/null
+++ b/Engine/Source/Runtime/FieldNotification/Public/FieldNotificationHelpers.h
@@ -0,0 +1,19 @@
+// Copyright Epic Games, Inc. All Rights Reserved.
+
+#pragma once
+
+#include "CoreMinimal.h"
+#include "UObject/Class.h"
+
+namespace UE::FieldNotification::Helpers
+{
+	UE_NODISCARD FORCEINLINE bool IsValidAsField(const UFunction* InFunction)
+	{
+		return InFunction != nullptr
+			&& !InFunction->HasAnyFunctionFlags(FUNC_Net | FUNC_Event)
+			&& InFunction->HasAllFunctionFlags(FUNC_BlueprintCallable | FUNC_Const)
+			&& InFunction->NumParms == 1
+			&& InFunction->GetReturnProperty() != nullptr;
+	}
+
+} //namespace
diff --git a/Engine/Source/Runtime/FieldNotification/Public/FieldNotificationId.h b/Engine/Source/Runtime/FieldNotification/Public/FieldNotificationId.h
new file mode 100644
--- /dev/null
+++ b/Engine/Source/Runtime/FieldNotification/Public/FieldNotificationId.h
@@ -0,0 +1,127 @@
+// Copyright Epic Games, Inc. All Rights Reserved.
+
+#pragma once
+
+#include "CoreMinimal.h"
+#include "FieldNotificationVariant.h"
+#include "UObject/Object.h"
+
+#include "FieldNotificationId.generated.h"
+
+
+namespace UE::FieldNotification
+{
+	
+struct FFieldId
+{
+public:
+	FFieldId() = default;
+
+	explicit FFieldId(FName InFieldName, int32 InBitNumber)
+		: FieldName(InFieldName)
+		, BitNumber(InBitNumber)
+	{
+	}
+
+
+public:
+	bool IsValid() const
+	{
+		return !FieldName.IsNone();
+	}
+
+	int32 GetIndex() const
+	{
+		return BitNumber;
+	}
+	
+	FName GetName() const
+	{
+		return FieldName;
+	}
+
+	FFieldVariant ToVariant(UObject* InContainer) const
+	{
+		if (InContainer && IsValid())
+		{
+			if (UFunction* Function = InContainer->GetClass()->FindFunctionByName(FieldName))
+			{
+				return FFieldVariant(Function);
+			}
+			else if (FProperty* Property = InContainer->GetClass()->FindPropertyByName(FieldName))
+			{
+				return FFieldVariant(Property);
+			}
+		}
+		return FFieldVariant();
+	}
+
+	bool operator==(const FFieldId& Other) const
+	{
+		return FieldName == Other.FieldName;
+	}
+
+	bool operator!=(const FFieldId& Other) const
+	{
+		return FieldName != Other.FieldName;
+	}
+
+	friend int32 GetTypeHash(const FFieldId& Value)
+	{
+		return GetTypeHash(Value.FieldName);
+	}
+
+private:
+	/** Name of the field. It can be a FProperty or UFunction. */
+	FName FieldName;
+	/** The bit this field is linked to. */
+	int32 BitNumber;
+};
+
+} //namespace
+
+
+USTRUCT(BlueprintType)
+struct FFieldNotificationId
+{
+	GENERATED_BODY()
+
+public:
+	FFieldNotificationId() = default;
+
+	explicit FFieldNotificationId(FName InFieldName)
+		: FieldName(InFieldName)
+	{
+
+	}
+
+public:
+	bool IsValid() const
+	{
+		return !FieldName.IsNone();
+	}
+
+	FName GetFieldName() const
+	{
+		return FieldName;
+	}
+
+	bool operator== (const FFieldNotificationId& Other) const
+	{
+		return FieldName == Other.FieldName;
+	}
+
+	bool operator!= (const FFieldNotificationId& Other) const
+	{
+		return FieldName != Other.FieldName;
+	}
+
+	friend int32 GetTypeHash(const FFieldNotificationId& Value)
+	{
+		return GetTypeHash(Value.FieldName);
+	}
+
+public:
+	UPROPERTY(BlueprintReadWrite, EditAnywhere, Category = "FieldNotify")
+	FName FieldName;
+};
diff --git a/Engine/Source/Runtime/FieldNotification/Public/FieldNotificationVariant.h b/Engine/Source/Runtime/FieldNotification/Public/FieldNotificationVariant.h
new file mode 100644
--- /dev/null
+++ b/Engine/Source/Runtime/FieldNotification/Public/FieldNotificationVariant.h
@@ -0,0 +1,70 @@
+// Copyright Epic Games, Inc. All Rights Reserved.
+
+#pragma once
+
+#include "CoreMinimal.h"
+#include "UObject/Class.h"
+#include "UObject/Field.h"
+#include "UObject/UnrealType.h"
+
+namespace UE::FieldNotification
+{
+
+struct FFieldVariant : private ::FFieldVariant
+{
+public:
+	FFieldVariant() = default;
+
+	explicit FFieldVariant(const FProperty* InProperty)
+		: ::FFieldVariant(InProperty)
+	{
+	}
+
+	explicit FFieldVariant(const UFunction* InFunction)
+		: ::FFieldVariant(InFunction)
+	{
+	}
+
+public:
+	using ::FFieldVariant::IsValid;
+	using ::FFieldVariant::GetOwnerClass;
+	using ::FFieldVariant::GetFName;
+	using ::FFieldVariant::operator bool;
+	using ::FFieldVariant::operator ==;
+	using ::FFieldVariant::operator !=;
+#if WITH_EDITORONLY_DATA
+	using ::FFieldVariant::HasMetaData;
+#endif
+
+	inline bool IsProperty() const
+	{
+		return IsValid() && !IsUObject();
+	}
+
+	inline bool IsFunction() const
+	{
+		return IsValid() && IsUObject();
+	}
+
+	inline FProperty* GetProperty()
+	{
+		return IsProperty() ? static_cast<FProperty*>(ToFieldUnsafe()) : nullptr;
+	}
+	
+	inline const FProperty* GetProperty() const
+	{
+		return IsProperty() ? static_cast<const FProperty*>(ToFieldUnsafe()) : nullptr;
+	}
+
+	inline UFunction* GetFunction()
+	{
+		return IsFunction() ? static_cast<UFunction*>(ToUObjectUnsafe()) : nullptr;
+	}
+
+	inline const UFunction* GetFunction() const
+	{
+		return IsFunction() ? static_cast<const UFunction*>(ToUObjectUnsafe()) : nullptr;
+	}
+};
+
+} //namespace
diff --git a/Engine/Source/Runtime/FieldNotification/Public/IFieldNotificationClassDescriptor.h b/Engine/Source/Runtime/FieldNotification/Public/IFieldNotificationClassDescriptor.h
new file mode 100644
--- /dev/null
+++ b/Engine/Source/Runtime/FieldNotification/Public/IFieldNotificationClassDescriptor.h
@@ -0,0 +1,44 @@
+// Copyright Epic Games, Inc. All Rights Reserved.
+
+#pragma once
+
+#include "CoreMinimal.h"
+#include "FieldNotificationId.h"
+
+namespace UE::FieldNotification
+{
+
+struct IClassDescriptor
+{
+	enum
+	{
+		Max_IndexOf_ = 0,
+	};
+
+	/** Find the FieldId by name. Returns an invalid FFIeldId if not found. */
+	FFieldId GetField(const UClass* Class, FName InFieldName) const
+	{
+		FFieldId FoundId;
+		ForEachField(Class, [&FoundId, InFieldName](const FFieldId Other)
+		{
+			if (Other.GetName() == InFieldName)
+			{
+				FoundId = Other;
+				return false;
+			}
+			return true;
+		});
+		return FoundId;
+	}
+
+	/** Execute the callback for every FieldId in the ClassDescriptor. */
+	virtual void ForEachField(const UClass* Class, TFunctionRef<bool(FFieldId FielId)> Callback) const
+	{
+
+	}
+
+	/** */
+	virtual ~IClassDescriptor() = default;
+};
+
+} // namespace
\ No newline at end of file
diff --git a/Engine/Source/Runtime/FieldNotification/Public/INotifyFieldValueChanged.h b/Engine/Source/Runtime/FieldNotification/Public/INotifyFieldValueChanged.h
new file mode 100644
--- /dev/null
+++ b/Engine/Source/Runtime/FieldNotification/Public/INotifyFieldValueChanged.h
@@ -0,0 +1,45 @@
+// Copyright Epic Games, Inc. All Rights Reserved.
+
+#pragma once
+
+#include "CoreMinimal.h"
+#include "FieldNotificationId.h"
+#include "IFieldNotificationClassDescriptor.h"
+#include "UObject/Interface.h"
+
+#include "INotifyFieldValueChanged.generated.h"
+
+DECLARE_DYNAMIC_DELEGATE_TwoParams(FFieldValueChangedDynamicDelegate, UObject*, Object, FFieldNotificationId, Field);
+
+
+UINTERFACE(MinimalAPI, NotBlueprintable)
+class UNotifyFieldValueChanged : public UInterface
+{
+	GENERATED_BODY()
+};
+
+
+class INotifyFieldValueChanged : public IInterface
+{
+	GENERATED_BODY()
+
+public:
+	DECLARE_MULTICAST_DELEGATE_TwoParams(FFieldValueChangedMulticastDelegate, UObject*, UE::FieldNotification::FFieldId);
+	using FFieldValueChangedDelegate = FFieldValueChangedMulticastDelegate::FDelegate;
+
+public:
+	/** Add a delegate that will be notified when the FieldId is value changed. */
+	virtual FDelegateHandle AddFieldValueChangedDelegate(UE::FieldNotification::FFieldId InFieldId, FFieldValueChangedDelegate InNewDelegate) = 0;
+
+	/** Remove a delegate that was added. */
+	virtual bool RemoveFieldValueChangedDelegate(UE::FieldNotification::FFieldId InFieldId, FDelegateHandle InHandle) = 0;
+
+	/** Remove all the delegate that are bound to the specified UserObject. */
+	virtual int32 RemoveAllFieldValueChangedDelegates(const void* InUserObject) = 0;
+
+	/** Remove all the delegate that are bound to the specified Field and UserObject. */
+	virtual int32 RemoveAllFieldValueChangedDelegates(UE::FieldNotification::FFieldId InFieldId, const void* InUserObject) = 0;
+
+	/** @returns the list of all the field that can notify when their value changes. */
+	virtual const UE::FieldNotification::IClassDescriptor& GetFieldNotificationDescriptor() const = 0;
+};
diff --git a/Engine/Source/Runtime/UMG/Private/Binding/WidgetFieldNotificationExtension.h b/Engine/Source/Runtime/UMG/Private/Binding/WidgetFieldNotificationExtension.h
--- a/Engine/Source/Runtime/UMG/Private/Binding/WidgetFieldNotificationExtension.h
+++ b/Engine/Source/Runtime/UMG/Private/Binding/WidgetFieldNotificationExtension.h
@@ -5,7 +5,7 @@
 #include "CoreMinimal.h"
 #include "Components/Visual.h"
 #include "Extensions/UserWidgetExtension.h"
-#include "FieldNotification/FieldMulticastDelegate.h"
+#include "FieldNotificationDelegate.h"
 
 #include "WidgetFieldNotificationExtension.generated.h"
 
diff --git a/Engine/Source/Runtime/UMG/Private/FieldNotification/FieldMulticastDelegate.cpp b/Engine/Source/Runtime/UMG/Private/FieldNotification/FieldMulticastDelegate.cpp
deleted file mode 100644
diff --git a/Engine/Source/Runtime/UMG/Private/Tests/FieldNotificationMulticastDelegateTest.cpp b/Engine/Source/Runtime/UMG/Private/Tests/FieldNotificationMulticastDelegateTest.cpp
deleted file mode 100644
diff --git a/Engine/Source/Runtime/UMG/Public/Blueprint/WidgetBlueprintGeneratedClass.h b/Engine/Source/Runtime/UMG/Public/Blueprint/WidgetBlueprintGeneratedClass.h
--- a/Engine/Source/Runtime/UMG/Public/Blueprint/WidgetBlueprintGeneratedClass.h
+++ b/Engine/Source/Runtime/UMG/Public/Blueprint/WidgetBlueprintGeneratedClass.h
@@ -6,7 +6,7 @@
 #include "UObject/ObjectMacros.h"
 #include "Binding/DynamicPropertyPath.h"
 #include "Engine/BlueprintGeneratedClass.h"
-#include "FieldNotification/FieldId.h"
+#include "FieldNotificationId.h"
 
 #include "WidgetBlueprintGeneratedClass.generated.h"
 
diff --git a/Engine/Source/Runtime/UMG/Public/Components/Widget.h b/Engine/Source/Runtime/UMG/Public/Components/Widget.h
--- a/Engine/Source/Runtime/UMG/Public/Components/Widget.h
+++ b/Engine/Source/Runtime/UMG/Public/Components/Widget.h
@@ -4,8 +4,8 @@
 
 #include "CoreMinimal.h"
 #include "UObject/ObjectMacros.h"
-#include "FieldNotification/FieldNotificationDeclaration.h"
-#include "FieldNotification/IFieldValueChanged.h"
+#include "FieldNotificationDeclaration.h"
+#include "INotifyFieldValueChanged.h"
 #include "Misc/Attribute.h"
 #include "Templates/SubclassOf.h"
 #include "UObject/ScriptMacros.h"
diff --git a/Engine/Source/Runtime/UMG/Public/FieldNotification/FieldId.h b/Engine/Source/Runtime/UMG/Public/FieldNotification/FieldId.h
--- a/Engine/Source/Runtime/UMG/Public/FieldNotification/FieldId.h
+++ b/Engine/Source/Runtime/UMG/Public/FieldNotification/FieldId.h
@@ -1,127 +1,16 @@
 // Copyright Epic Games, Inc. All Rights Reserved.
-
 #pragma once
 
-#include "CoreMinimal.h"
-#include "FieldNotification/FieldVariant.h"
-#include "UObject/Object.h"
-
-#include "FieldId.generated.h"
+// HEADER_UNIT_SKIP - Deprecated
 
+#include "FieldNotificationId.h"
 
-namespace UE::FieldNotification
-{
-	
-struct FFieldId
+UE_DEPRECATED(5.3, "Including FieldNotification/FieldId.h is deprecated. The file location changed. Please include FieldNotificationId.h instead and add the module FieldNotification to your .build.cs")
+inline void FieldIdHeaderDeprecatedWarning()
 {
-public:
-	FFieldId() = default;
-
-	explicit FFieldId(FName InFieldName, int32 InBitNumber)
-		: FieldName(InFieldName)
-		, BitNumber(InBitNumber)
-	{
-	}
-
-
-public:
-	bool IsValid() const
-	{
-		return !FieldName.IsNone();
-	}
-
-	int32 GetIndex() const
-	{
-		return BitNumber;
-	}
-	
-	FName GetName() const
-	{
-		return FieldName;
-	}
-
-	FFieldVariant ToVariant(UObject* InContainer) const
-	{
-		if (InContainer && IsValid())
-		{
-			if (UFunction* Function = InContainer->GetClass()->FindFunctionByName(FieldName))
-			{
-				return FFieldVariant(Function);
-			}
-			else if (FProperty* Property = InContainer->GetClass()->FindPropertyByName(FieldName))
-			{
-				return FFieldVariant(Property);
-			}
-		}
-		return FFieldVariant();
-	}
-
-	bool operator==(const FFieldId& Other) const
-	{
-		return FieldName == Other.FieldName;
-	}
-
-	bool operator!=(const FFieldId& Other) const
-	{
-		return FieldName != Other.FieldName;
-	}
-
-	friend int32 GetTypeHash(const FFieldId& Value)
-	{
-		return GetTypeHash(Value.FieldName);
-	}
-
-private:
-	/** Name of the field. It can be a FProperty or UFunction. */
-	FName FieldName;
-	/** The bit this field is linked to. */
-	int32 BitNumber;
-};
+}
 
-} //namespace
-
-
-USTRUCT(BlueprintType)
-struct FFieldNotificationId
+inline void TriggerFieldIdHeaderDeprecatedWarning()
 {
-	GENERATED_BODY()
-
-public:
-	FFieldNotificationId() = default;
-
-	explicit FFieldNotificationId(FName InFieldName)
-		: FieldName(InFieldName)
-	{
-
-	}
-
-public:
-	bool IsValid() const
-	{
-		return !FieldName.IsNone();
-	}
-
-	FName GetFieldName() const
-	{
-		return FieldName;
-	}
-
-	bool operator== (const FFieldNotificationId& Other) const
-	{
-		return FieldName == Other.FieldName;
-	}
-
-	bool operator!= (const FFieldNotificationId& Other) const
-	{
-		return FieldName != Other.FieldName;
-	}
-
-	friend int32 GetTypeHash(const FFieldNotificationId& Value)
-	{
-		return GetTypeHash(Value.FieldName);
-	}
-
-public:
-	UPROPERTY(BlueprintReadWrite, EditAnywhere, Category = "FieldNotify")
-	FName FieldName;
-};
+	FieldIdHeaderDeprecatedWarning();
+}
diff --git a/Engine/Source/Runtime/UMG/Public/FieldNotification/FieldMulticastDelegate.h b/Engine/Source/Runtime/UMG/Public/FieldNotification/FieldMulticastDelegate.h
--- a/Engine/Source/Runtime/UMG/Public/FieldNotification/FieldMulticastDelegate.h
+++ b/Engine/Source/Runtime/UMG/Public/FieldNotification/FieldMulticastDelegate.h
@@ -1,158 +1,16 @@
 // Copyright Epic Games, Inc. All Rights Reserved.
-
 #pragma once
 
-#include "Algo/BinarySearch.h"
-#include "Containers/Array.h"
-#include "Containers/BitArray.h"
-#include "CoreMinimal.h"
-#include "Delegates/Delegate.h"
-#include "FieldNotification/FieldId.h"
-#include "FieldNotification/IFieldValueChanged.h"
-#include "HAL/Platform.h"
-#include "Templates/IdentityFunctor.h"
-#include "Templates/Less.h"
-#include "Templates/UnrealTemplate.h"
-#include "UObject/NameTypes.h"
-#include "UObject/Object.h"
-#include "UObject/WeakObjectPtr.h"
-#include "UObject/WeakObjectPtrTemplates.h"
-
-#include <limits>
+// HEADER_UNIT_SKIP - Deprecated
 
-class UObject;
+#include "FieldNotificationDelegate.h"
 
-namespace UE::FieldNotification
-{
-	
-class UMG_API FFieldMulticastDelegate
+UE_DEPRECATED(5.3, "Including FieldNotification/FieldMulticastDelegate.h is deprecated. The file location changed. Please include FieldNotificationDelegate.h instead and add the module FieldNotification to your .build.cs")
+inline void FieldMulticastDelegateHeaderDeprecatedWarning()
 {
-public:
-	using FDelegate = INotifyFieldValueChanged::FFieldValueChangedDelegate;
-	using FDynamicDelegate = FFieldValueChangedDynamicDelegate;
-
-private:
-	struct FInstanceProtectedDelegate : public TDelegateBase<FDefaultDelegateUserPolicy>
-	{
-		using Super = TDelegateBase<FDefaultDelegateUserPolicy>;
-		using Super::GetDelegateInstanceProtected;
-	};
-
-	struct FInvocationKey
-	{
-		TWeakObjectPtr<const UObject> Object;
-		FFieldId Id;
-		FName DynamicName;
-		bool operator<(const FInvocationKey& Element) const
-		{
-			return Id.GetName().FastLess(Element.Id.GetName());
-		}
-	};
-
-	struct FInvocationElement
-	{
-		FInvocationKey Key;
-		FDelegate Delegate;
-		bool operator<(const FInvocationElement& InElement) const
-		{
-			return Key < InElement.Key;
-		}
-	};
-
-	friend bool operator<(const FInvocationElement& A, const FInvocationKey& B)
-	{
-		return A.Key < B;
-	}
-	friend bool operator<(const FInvocationKey& A, const FInvocationElement& B)
-	{
-		return A < B.Key;
-	}
-	friend bool operator<(const FInvocationElement& A, const FFieldId& B)
-	{
-		return A.Key.Id.GetName().FastLess(B.GetName());
-	}
-	friend bool operator<(const FFieldId& A, const FInvocationElement& B)
-	{
-		return A.GetName().FastLess(B.Key.Id.GetName());
-	}
-
-	using InvocationListType = TArray<FInvocationElement>;
-
-public:
-	FDelegateHandle Add(const UObject* InObject, FFieldId InFieldId, FDelegate InNewDelegate);
-	FDelegateHandle Add(const UObject* InObject, FFieldId InFieldId, const FDynamicDelegate& InDynamicDelegate);
-
-	struct FRemoveResult
-	{
-		bool bRemoved = false;
-		bool bHasOtherBoundDelegates = false;
-		const UObject* Object = nullptr;
-		FFieldId FieldId;
-	};
-	FRemoveResult Remove(FDelegateHandle InDelegate);
-	FRemoveResult Remove(const FDynamicDelegate& InDynamicDelegate);
-
-	struct FRemoveFromResult
-	{
-		bool bRemoved = false;
-		bool bHasOtherBoundDelegates = false;
-	};
-	FRemoveFromResult RemoveFrom(const UObject* InObject, FFieldId InFieldId, FDelegateHandle InDelegate);
-	FRemoveFromResult RemoveFrom(const UObject* InObject, FFieldId InFieldId, const FDynamicDelegate& InDynamicDelegate);
+}
 
-	struct FRemoveAllResult
-	{
-		int32 RemoveCount = 0;
-		TBitArray<> HasFields;
-	};
-	FRemoveAllResult RemoveAll(const UObject* InObject, const void* InUserObject);
-	FRemoveAllResult RemoveAll(const UObject* InObject, FFieldId InFieldId, const void* InUserObject);
-
-	void Broadcast(UObject* InObject, FFieldId InFieldId);
-
-	void Reset();
-
-private:
-	int32 LowerBound(FFieldId InFieldId) const
-	{
-		int32 Num = AddedEmplaceAt > GetNum(Delegates) ? GetNum(Delegates) : AddedEmplaceAt;
-		return AlgoImpl::LowerBoundInternal(GetData(Delegates), Num, InFieldId, FIdentityFunctor(), TLess<>());
-	}
-
-	int32 LowerBound(const FInvocationKey& InKey) const
-	{
-		int32 Num = AddedEmplaceAt > GetNum(Delegates) ? GetNum(Delegates) : AddedEmplaceAt;
-		return AlgoImpl::LowerBoundInternal(GetData(Delegates), Num, InKey, FIdentityFunctor(), TLess<>());
-	}
-
-	int32 UpperBound(FFieldId InFieldId) const
-	{
-		int32 Num = AddedEmplaceAt > GetNum(Delegates) ? GetNum(Delegates) : AddedEmplaceAt;
-		return AlgoImpl::UpperBoundInternal(GetData(Delegates), Num, InFieldId, FIdentityFunctor(), TLess<>());
-	}
-
-	int32 UpperBound(const FInvocationKey& InKey) const
-	{
-		int32 Num = AddedEmplaceAt > GetNum(Delegates) ? GetNum(Delegates) : AddedEmplaceAt;
-		return AlgoImpl::UpperBoundInternal(GetData(Delegates), Num, InKey, FIdentityFunctor(), TLess<>());
-	}
-
-	const IDelegateInstance* GetDelegateInstance(FDelegate& InDelegate) const
-	{
-		return static_cast<FInstanceProtectedDelegate&>(static_cast<TDelegateBase<FDefaultDelegateUserPolicy>&>(InDelegate)).GetDelegateInstanceProtected();
-	}
-
-	void ExecuteLockOperations();
-	void RemoveElement(FInvocationElement& Element, int32 Index, FRemoveResult& Result);
-	void CompleteRemove(FRemoveResult& Result) const;
-	template<typename TRemoveOrUnbind>
-	void RemoveFrom(TRemoveOrUnbind& RemoveOrUnbind, FFieldId FieldId, bool& bFieldPresent);
-
-private:
-	InvocationListType Delegates;
-	int16 DelegateLockCount = 0;
-	int16 CompactionCount = 0;
-	uint16 AddedEmplaceAt = std::numeric_limits<uint16>::max();
-};
-
-} //namespace
\ No newline at end of file
+inline void TriggerFieldMulticastDelegateHeaderDeprecatedWarning()
+{
+	FieldMulticastDelegateHeaderDeprecatedWarning();
+}
diff --git a/Engine/Source/Runtime/UMG/Public/FieldNotification/FieldNotificationDeclaration.h b/Engine/Source/Runtime/UMG/Public/FieldNotification/FieldNotificationDeclaration.h
--- a/Engine/Source/Runtime/UMG/Public/FieldNotification/FieldNotificationDeclaration.h
+++ b/Engine/Source/Runtime/UMG/Public/FieldNotification/FieldNotificationDeclaration.h
@@ -1,244 +1,16 @@
 // Copyright Epic Games, Inc. All Rights Reserved.
-
 #pragma once
 
-#include "FieldNotification/FieldId.h"
-#include "FieldNotification/IClassDescriptor.h"
-
-
-/*
- *	struct FFieldNotificationClassDescriptor : Super::FFieldNotificationClassDescriptor
- *	{
- *	private:
- *		using SuperDescriptor = Super::FFieldNotificationClassDescriptor;
- *		static const ::UE::FieldNotification::FFieldId* AllFields[];
- *		friend ThisClass;
- *	public:
- *		static const ::UE::FieldNotification::FFieldId FieldA;
- *		static const ::UE::FieldNotification::FFieldId FieldB;
- *		enum
- *		{
- *			IndexOf_FieldA = Super::FFieldNotificationClassDescriptor::Max_IndexOf_ + 0,
- *			IndexOf_FieldB,
- *			Max_IndexOf_,
- *		};
- * 	virtual void ForEachField(const UClass* Class, TFunctionRef<bool(::UE::FieldNotification::FFieldId FielId)> Callback) const override;
- *	};
- * virtual const ::UE::FieldNotification::IClassDescriptor& GetFieldNotificationDescriptor() const override;
- */
-
-#define UE_FIELD_NOTIFICATION_OVERRIDE_GET_FIELD_NOTIFICATION_DESCRIPTOR() \
-	virtual const ::UE::FieldNotification::IClassDescriptor& GetFieldNotificationDescriptor() const \
-	{ \
- 		static FFieldNotificationClassDescriptor Instance; \
- 		return Instance; \
- 	}
-
-#define UE_FIELD_NOTIFICATION_DECLARE_CLASS_DESCRIPTOR_BEGIN(API_STRING) \
-	struct API_STRING FFieldNotificationClassDescriptor : public Super::FFieldNotificationClassDescriptor \
-	{ \
-	private: \
-		using SuperDescriptor = Super::FFieldNotificationClassDescriptor; \
-		static const ::UE::FieldNotification::FFieldId* AllFields[]; \
-		friend ThisClass; \
-	public:
-
-
-#define UE_FIELD_NOTIFICATION_DECLARE_CLASS_DESCRIPTOR_END() \
-		virtual void ForEachField(const UClass* Class, TFunctionRef<bool(::UE::FieldNotification::FFieldId FielId)> Callback) const override \
-		{ \
-			for (int32 Index = 0; Index < Max_IndexOf_-SuperDescriptor::Max_IndexOf_; ++Index) \
-			{ \
-				if (!Callback(*AllFields[Index])) \
-				{ \
-					return; \
-				} \
-			} \
-			SuperDescriptor::ForEachField(Class, Callback); \
-		} \
-	}; \
-	UE_FIELD_NOTIFICATION_OVERRIDE_GET_FIELD_NOTIFICATION_DESCRIPTOR()
-
-
-#define UE_FIELD_NOTIFICATION_DECLARE_CLASS_DESCRIPTOR_BASE_BEGIN(API_STRING) \
-	struct API_STRING FFieldNotificationClassDescriptor : public ::UE::FieldNotification::IClassDescriptor \
-	{ \
-	private: \
-		using SuperDescriptor = ::UE::FieldNotification::IClassDescriptor; \
-		static const ::UE::FieldNotification::FFieldId* AllFields[]; \
-		friend ThisClass; \
-	public:
-
-
-#define UE_FIELD_NOTIFICATION_DECLARE_CLASS_DESCRIPTOR_BASE_END() \
-		virtual void ForEachField(const UClass* Class, TFunctionRef<bool(::UE::FieldNotification::FFieldId FielId)> Callback) const override; \
-	}; \
-	UE_FIELD_NOTIFICATION_OVERRIDE_GET_FIELD_NOTIFICATION_DESCRIPTOR()
-
-
-#define UE_FIELD_NOTIFICATION_DECLARE_ENUM_FIELD_BEGIN(Name) \
-		enum EField \
-		{ \
-			IndexOf_##Name = SuperDescriptor::Max_IndexOf_ + 0,
-
-
-#define UE_FIELD_NOTIFICATION_DECLARE_ENUM_FIELD_END() \
-			Max_IndexOf_, \
-		};
-
-
-#define UE_FIELD_NOTIFICATION_DECLARE_FIELD(Name) \
-	static const ::UE::FieldNotification::FFieldId Name;
-
-
-#define UE_FIELD_NOTIFICATION_DECLARE_ENUM_FIELD(Name) \
-			IndexOf_##Name,
-
-
-#define UE_FIELD_NOTIFICATION_IMPLEMENTATION_BEGIN(ClassName) \
-	const ::UE::FieldNotification::FFieldId* ClassName::FFieldNotificationClassDescriptor::AllFields[] = {
-
-
-#define UE_FIELD_NOTIFICATION_IMPLEMENTATION_END(ClassName) }
-
-//#if DO_CHECK
-//#define UE_FIELD_NOTIFICATION_IMPLEMENTATION_END(ClassName) \
-//	}; \
-//	const ::UE::FieldNotification::IClassDescriptor& ClassName::GetFieldNotificationDescriptor() const \
-//	{ \
-//		static_assert(UE_ARRAY_COUNT(ClassName::FFieldNotificationClassDescriptor::AllFields) == ClassName::FFieldNotificationClassDescriptor::Max_IndexOf_-ClassName::FFieldNotificationClassDescriptor::SuperDescriptor::Max_IndexOf_, "The descriptor for class " #ClassName " doesn't not implement the same number of field as declared in in the constructor."); \
-//		struct FLocal \
-//		{ \
-//			ClassName::FFieldNotificationClassDescriptor Instance; \
-//			FLocal() \
-//			{ \
-//				int32 NumberOfField = Instance.GetNumberOfField(); \
-//				for (int32 Index = 0; Index < NumberOfField; ++Index) \
-//				{ \
-//					::UE::FieldNotification::FFieldId Id = Instance.GetField(Index); \
-//					ensureAlwaysMsgf(Id.IsValid(), TEXT("The FFieldId '%d' for the class '%s' is invalid"), Index, *ClassName::StaticClass()->GetName()); \
-//					ensureAlwaysMsgf(Id.GetIndex() == Index, TEXT("The FFieldId at index '%d' doesn't match the id index '%d'."), Index, Id.GetIndex()); \
-//				} \
-//			} \
-//		}; \
-//		static FLocal Local; \
-//		return Local.Instance; \
-//	}
-//#else
-//#define UE_FIELD_NOTIFICATION_IMPLEMENTATION_END(ClassName) \
-//	}; \
-//	const ::UE::FieldNotification::IClassDescriptor& ClassName::GetFieldNotificationDescriptor() const \
-//	{ \
-//		static FFieldNotificationClassDescriptor Instance; \
-//		return Instance; \
-//	}
-//#endif
-
-
-#define UE_FIELD_NOTIFICATION_IMPLEMENT_FIELD(ClassName, MemberName) const ::UE::FieldNotification::FFieldId ClassName::FFieldNotificationClassDescriptor::MemberName( FName(TEXT(#MemberName)), ClassName::FFieldNotificationClassDescriptor::IndexOf_##MemberName );
-#define UE_FIELD_NOTIFICATION_IMPLEMENT_ENUM_FIELD(ClassName, MemberName) &ClassName::FFieldNotificationClassDescriptor::MemberName,
-
-
-#define UE_FIELD_NOTIFICATION_IMPLEMENT_CLASS_DESCRIPTOR_OneField(ClassName, Field1) \
-	UE_FIELD_NOTIFICATION_IMPLEMENT_FIELD(ClassName, Field1) \
-	UE_FIELD_NOTIFICATION_IMPLEMENTATION_BEGIN(ClassName) \
-	UE_FIELD_NOTIFICATION_IMPLEMENT_ENUM_FIELD(ClassName, Field1) \
-	UE_FIELD_NOTIFICATION_IMPLEMENTATION_END(ClassName)
-
-#define UE_FIELD_NOTIFICATION_IMPLEMENT_CLASS_DESCRIPTOR_TwoFields(ClassName, Field1, Field2) \
-	UE_FIELD_NOTIFICATION_IMPLEMENT_FIELD(ClassName, Field1) \
-	UE_FIELD_NOTIFICATION_IMPLEMENT_FIELD(ClassName, Field2) \
-	UE_FIELD_NOTIFICATION_IMPLEMENTATION_BEGIN(ClassName) \
-	UE_FIELD_NOTIFICATION_IMPLEMENT_ENUM_FIELD(ClassName, Field1) \
-	UE_FIELD_NOTIFICATION_IMPLEMENT_ENUM_FIELD(ClassName, Field2) \
-	UE_FIELD_NOTIFICATION_IMPLEMENTATION_END(ClassName)
-
-#define UE_FIELD_NOTIFICATION_IMPLEMENT_CLASS_DESCRIPTOR_ThreeFields(ClassName, Field1, Field2, Field3) \
-	UE_FIELD_NOTIFICATION_IMPLEMENT_FIELD(ClassName, Field1) \
-	UE_FIELD_NOTIFICATION_IMPLEMENT_FIELD(ClassName, Field2) \
-	UE_FIELD_NOTIFICATION_IMPLEMENT_FIELD(ClassName, Field3) \
-	UE_FIELD_NOTIFICATION_IMPLEMENTATION_BEGIN(ClassName) \
-	UE_FIELD_NOTIFICATION_IMPLEMENT_ENUM_FIELD(ClassName, Field1) \
-	UE_FIELD_NOTIFICATION_IMPLEMENT_ENUM_FIELD(ClassName, Field2) \
-	UE_FIELD_NOTIFICATION_IMPLEMENT_ENUM_FIELD(ClassName, Field3) \
-	UE_FIELD_NOTIFICATION_IMPLEMENTATION_END(ClassName)
-
-#define UE_FIELD_NOTIFICATION_IMPLEMENT_CLASS_DESCRIPTOR_FourFields(ClassName, Field1, Field2, Field3, Field4) \
-	UE_FIELD_NOTIFICATION_IMPLEMENT_FIELD(ClassName, Field1) \
-	UE_FIELD_NOTIFICATION_IMPLEMENT_FIELD(ClassName, Field2) \
-	UE_FIELD_NOTIFICATION_IMPLEMENT_FIELD(ClassName, Field3) \
-	UE_FIELD_NOTIFICATION_IMPLEMENT_FIELD(ClassName, Field4) \
-	UE_FIELD_NOTIFICATION_IMPLEMENTATION_BEGIN(ClassName) \
-	UE_FIELD_NOTIFICATION_IMPLEMENT_ENUM_FIELD(ClassName, Field1) \
-	UE_FIELD_NOTIFICATION_IMPLEMENT_ENUM_FIELD(ClassName, Field2) \
-	UE_FIELD_NOTIFICATION_IMPLEMENT_ENUM_FIELD(ClassName, Field3) \
-	UE_FIELD_NOTIFICATION_IMPLEMENT_ENUM_FIELD(ClassName, Field4) \
-	UE_FIELD_NOTIFICATION_IMPLEMENTATION_END(ClassName)
-
-#define UE_FIELD_NOTIFICATION_IMPLEMENT_CLASS_DESCRIPTOR_FiveFields(ClassName, Field1, Field2, Field3, Field4, Field5) \
-	UE_FIELD_NOTIFICATION_IMPLEMENT_FIELD(ClassName, Field1) \
-	UE_FIELD_NOTIFICATION_IMPLEMENT_FIELD(ClassName, Field2) \
-	UE_FIELD_NOTIFICATION_IMPLEMENT_FIELD(ClassName, Field3) \
-	UE_FIELD_NOTIFICATION_IMPLEMENT_FIELD(ClassName, Field4) \
-	UE_FIELD_NOTIFICATION_IMPLEMENT_FIELD(ClassName, Field5) \
-	UE_FIELD_NOTIFICATION_IMPLEMENTATION_BEGIN(ClassName) \
-	UE_FIELD_NOTIFICATION_IMPLEMENT_ENUM_FIELD(ClassName, Field1) \
-	UE_FIELD_NOTIFICATION_IMPLEMENT_ENUM_FIELD(ClassName, Field2) \
-	UE_FIELD_NOTIFICATION_IMPLEMENT_ENUM_FIELD(ClassName, Field3) \
-	UE_FIELD_NOTIFICATION_IMPLEMENT_ENUM_FIELD(ClassName, Field4) \
-	UE_FIELD_NOTIFICATION_IMPLEMENT_ENUM_FIELD(ClassName, Field5) \
-	UE_FIELD_NOTIFICATION_IMPLEMENTATION_END(ClassName)
+// HEADER_UNIT_SKIP - Deprecated
 
-#define UE_FIELD_NOTIFICATION_IMPLEMENT_CLASS_DESCRIPTOR_SixFields(ClassName, Field1, Field2, Field3, Field4, Field5, Field6) \
-	UE_FIELD_NOTIFICATION_IMPLEMENT_FIELD(ClassName, Field1) \
-	UE_FIELD_NOTIFICATION_IMPLEMENT_FIELD(ClassName, Field2) \
-	UE_FIELD_NOTIFICATION_IMPLEMENT_FIELD(ClassName, Field3) \
-	UE_FIELD_NOTIFICATION_IMPLEMENT_FIELD(ClassName, Field4) \
-	UE_FIELD_NOTIFICATION_IMPLEMENT_FIELD(ClassName, Field5) \
-	UE_FIELD_NOTIFICATION_IMPLEMENT_FIELD(ClassName, Field6) \
-	UE_FIELD_NOTIFICATION_IMPLEMENTATION_BEGIN(ClassName) \
-	UE_FIELD_NOTIFICATION_IMPLEMENT_ENUM_FIELD(ClassName, Field1) \
-	UE_FIELD_NOTIFICATION_IMPLEMENT_ENUM_FIELD(ClassName, Field2) \
-	UE_FIELD_NOTIFICATION_IMPLEMENT_ENUM_FIELD(ClassName, Field3) \
-	UE_FIELD_NOTIFICATION_IMPLEMENT_ENUM_FIELD(ClassName, Field4) \
-	UE_FIELD_NOTIFICATION_IMPLEMENT_ENUM_FIELD(ClassName, Field5) \
-	UE_FIELD_NOTIFICATION_IMPLEMENT_ENUM_FIELD(ClassName, Field6) \
-	UE_FIELD_NOTIFICATION_IMPLEMENTATION_END(ClassName)
+#include "FieldNotificationDeclaration.h"
 
-#define UE_FIELD_NOTIFICATION_IMPLEMENT_CLASS_DESCRIPTOR_SevenFields(ClassName, Field1, Field2, Field3, Field4, Field5, Field6, Field7) \
-	UE_FIELD_NOTIFICATION_IMPLEMENT_FIELD(ClassName, Field1) \
-	UE_FIELD_NOTIFICATION_IMPLEMENT_FIELD(ClassName, Field2) \
-	UE_FIELD_NOTIFICATION_IMPLEMENT_FIELD(ClassName, Field3) \
-	UE_FIELD_NOTIFICATION_IMPLEMENT_FIELD(ClassName, Field4) \
-	UE_FIELD_NOTIFICATION_IMPLEMENT_FIELD(ClassName, Field5) \
-	UE_FIELD_NOTIFICATION_IMPLEMENT_FIELD(ClassName, Field6) \
-	UE_FIELD_NOTIFICATION_IMPLEMENT_FIELD(ClassName, Field7) \
-	UE_FIELD_NOTIFICATION_IMPLEMENTATION_BEGIN(ClassName) \
-	UE_FIELD_NOTIFICATION_IMPLEMENT_ENUM_FIELD(ClassName, Field1) \
-	UE_FIELD_NOTIFICATION_IMPLEMENT_ENUM_FIELD(ClassName, Field2) \
-	UE_FIELD_NOTIFICATION_IMPLEMENT_ENUM_FIELD(ClassName, Field3) \
-	UE_FIELD_NOTIFICATION_IMPLEMENT_ENUM_FIELD(ClassName, Field4) \
-	UE_FIELD_NOTIFICATION_IMPLEMENT_ENUM_FIELD(ClassName, Field5) \
-	UE_FIELD_NOTIFICATION_IMPLEMENT_ENUM_FIELD(ClassName, Field6) \
-	UE_FIELD_NOTIFICATION_IMPLEMENT_ENUM_FIELD(ClassName, Field7) \
-	UE_FIELD_NOTIFICATION_IMPLEMENTATION_END(ClassName)
+UE_DEPRECATED(5.3, "Including FieldNotification/FieldNotificationDeclaration.h is deprecated. The location of has changed. Please include FieldNotificationDeclaration.h instead and add the module FieldNotification to your .build.cs")
+inline void FieldNotificationDeclarationHeaderDeprecatedWarning()
+{
+}
 
-#define UE_FIELD_NOTIFICATION_IMPLEMENT_CLASS_DESCRIPTOR_EightFields(ClassName, Field1, Field2, Field3, Field4, Field5, Field6, Field7, Field8) \
-	UE_FIELD_NOTIFICATION_IMPLEMENT_FIELD(ClassName, Field1) \
-	UE_FIELD_NOTIFICATION_IMPLEMENT_FIELD(ClassName, Field2) \
-	UE_FIELD_NOTIFICATION_IMPLEMENT_FIELD(ClassName, Field3) \
-	UE_FIELD_NOTIFICATION_IMPLEMENT_FIELD(ClassName, Field4) \
-	UE_FIELD_NOTIFICATION_IMPLEMENT_FIELD(ClassName, Field5) \
-	UE_FIELD_NOTIFICATION_IMPLEMENT_FIELD(ClassName, Field6) \
-	UE_FIELD_NOTIFICATION_IMPLEMENT_FIELD(ClassName, Field7) \
-	UE_FIELD_NOTIFICATION_IMPLEMENT_FIELD(ClassName, Field8) \
-	UE_FIELD_NOTIFICATION_IMPLEMENTATION_BEGIN(ClassName) \
-	UE_FIELD_NOTIFICATION_IMPLEMENT_ENUM_FIELD(ClassName, Field1) \
-	UE_FIELD_NOTIFICATION_IMPLEMENT_ENUM_FIELD(ClassName, Field2) \
-	UE_FIELD_NOTIFICATION_IMPLEMENT_ENUM_FIELD(ClassName, Field3) \
-	UE_FIELD_NOTIFICATION_IMPLEMENT_ENUM_FIELD(ClassName, Field4) \
-	UE_FIELD_NOTIFICATION_IMPLEMENT_ENUM_FIELD(ClassName, Field5) \
-	UE_FIELD_NOTIFICATION_IMPLEMENT_ENUM_FIELD(ClassName, Field6) \
-	UE_FIELD_NOTIFICATION_IMPLEMENT_ENUM_FIELD(ClassName, Field7) \
-	UE_FIELD_NOTIFICATION_IMPLEMENT_ENUM_FIELD(ClassName, Field8) \
-	UE_FIELD_NOTIFICATION_IMPLEMENTATION_END(ClassName)
+inline void TriggerFieldNotificationDeclarationHeaderDeprecatedWarning()
+{
+	FieldNotificationDeclarationHeaderDeprecatedWarning();
+}
diff --git a/Engine/Source/Runtime/UMG/Public/FieldNotification/FieldNotificationHelpers.h b/Engine/Source/Runtime/UMG/Public/FieldNotification/FieldNotificationHelpers.h
--- a/Engine/Source/Runtime/UMG/Public/FieldNotification/FieldNotificationHelpers.h
+++ b/Engine/Source/Runtime/UMG/Public/FieldNotification/FieldNotificationHelpers.h
@@ -1,19 +1,16 @@
 // Copyright Epic Games, Inc. All Rights Reserved.
-
 #pragma once
 
-#include "CoreMinimal.h"
-#include "UObject/Class.h"
+// HEADER_UNIT_SKIP - Deprecated
+
+#include "FieldNotificationHelpers.h"
 
-namespace UE::FieldNotification::Helpers
+UE_DEPRECATED(5.3, "Including FieldNotification/FieldNotificationHelpers.h is deprecated. The file location changed. Please include FieldNotificationHelpers.h instead and add the module FieldNotification to your .build.cs")
+inline void FieldNotificationHelpersHeaderDeprecatedWarning()
 {
-	UE_NODISCARD FORCEINLINE bool IsValidAsField(const UFunction* InFunction)
-	{
-		return InFunction != nullptr
-			&& !InFunction->HasAnyFunctionFlags(FUNC_Net | FUNC_Event)
-			&& InFunction->HasAllFunctionFlags(FUNC_BlueprintCallable | FUNC_Const)
-			&& InFunction->NumParms == 1
-			&& InFunction->GetReturnProperty() != nullptr;
-	}
+}
 
-} //namespace
+inline void TriggerFieldNotificationHelpersHeaderDeprecatedWarning()
+{
+	FieldNotificationHelpersHeaderDeprecatedWarning();
+}
diff --git a/Engine/Source/Runtime/UMG/Public/FieldNotification/FieldVariant.h b/Engine/Source/Runtime/UMG/Public/FieldNotification/FieldVariant.h
--- a/Engine/Source/Runtime/UMG/Public/FieldNotification/FieldVariant.h
+++ b/Engine/Source/Runtime/UMG/Public/FieldNotification/FieldVariant.h
@@ -1,70 +1,16 @@
 // Copyright Epic Games, Inc. All Rights Reserved.
-
 #pragma once
 
-#include "CoreMinimal.h"
-#include "UObject/Class.h"
-#include "UObject/Field.h"
-#include "UObject/UnrealType.h"
+// HEADER_UNIT_SKIP - Deprecated
 
-namespace UE::FieldNotification
-{
+#include "FieldNotificationVariant.h"
 
-struct FFieldVariant : private ::FFieldVariant
+UE_DEPRECATED(5.3, "Including FieldNotification/FieldVariant.h is deprecated. The file location changed. Please include FieldNotificationVariant.h instead and add the module FieldNotification to your .build.cs")
+inline void FieldVariantHeaderDeprecatedWarning()
 {
-public:
-	FFieldVariant() = default;
-
-	explicit FFieldVariant(const FProperty* InProperty)
-		: ::FFieldVariant(InProperty)
-	{
-	}
-
-	explicit FFieldVariant(const UFunction* InFunction)
-		: ::FFieldVariant(InFunction)
-	{
-	}
-
-public:
-	using ::FFieldVariant::IsValid;
-	using ::FFieldVariant::GetOwnerClass;
-	using ::FFieldVariant::GetFName;
-	using ::FFieldVariant::operator bool;
-	using ::FFieldVariant::operator ==;
-	using ::FFieldVariant::operator !=;
-#if WITH_EDITORONLY_DATA
-	using ::FFieldVariant::HasMetaData;
-#endif
-
-	inline bool IsProperty() const
-	{
-		return IsValid() && !IsUObject();
-	}
+}
 
-	inline bool IsFunction() const
-	{
-		return IsValid() && IsUObject();
-	}
-
-	inline FProperty* GetProperty()
-	{
-		return IsProperty() ? static_cast<FProperty*>(ToFieldUnsafe()) : nullptr;
-	}
-	
-	inline const FProperty* GetProperty() const
-	{
-		return IsProperty() ? static_cast<const FProperty*>(ToFieldUnsafe()) : nullptr;
-	}
-
-	inline UFunction* GetFunction()
-	{
-		return IsFunction() ? static_cast<UFunction*>(ToUObjectUnsafe()) : nullptr;
-	}
-
-	inline const UFunction* GetFunction() const
-	{
-		return IsFunction() ? static_cast<const UFunction*>(ToUObjectUnsafe()) : nullptr;
-	}
-};
-
-} //namespace
+inline void TriggerFieldVariantHeaderDeprecatedWarning()
+{
+	FieldVariantHeaderDeprecatedWarning();
+}
diff --git a/Engine/Source/Runtime/UMG/Public/FieldNotification/IClassDescriptor.h b/Engine/Source/Runtime/UMG/Public/FieldNotification/IClassDescriptor.h
--- a/Engine/Source/Runtime/UMG/Public/FieldNotification/IClassDescriptor.h
+++ b/Engine/Source/Runtime/UMG/Public/FieldNotification/IClassDescriptor.h
@@ -1,44 +1,16 @@
 // Copyright Epic Games, Inc. All Rights Reserved.
-
 #pragma once
 
-#include "CoreMinimal.h"
-#include "FieldNotification/FieldId.h"
-
-namespace UE::FieldNotification
-{
-	
-	struct IClassDescriptor
-	{
-		enum
-		{
-			Max_IndexOf_ = 0,
-		};
-
-		/** Find the FieldId by name. Returns an invalid FFIeldId if not found. */
-		FFieldId GetField(const UClass* Class, FName InFieldName) const
-		{
-			FFieldId FoundId;
-			ForEachField(Class, [&FoundId, InFieldName](const FFieldId Other)
-			{
-				if (Other.GetName() == InFieldName)
-				{
-					FoundId = Other;
-					return false;
-				}
-				return true;
-			});
-			return FoundId;
-		}
+// HEADER_UNIT_SKIP - Deprecated
 
-		/** Execute the callback for every FieldId in the ClassDescriptor. */
-		virtual void ForEachField(const UClass* Class, TFunctionRef<bool(FFieldId FielId)> Callback) const
-		{
+#include "IFieldNotificationClassDescriptor.h"
 
-		}
-
-		/** */
-		virtual ~IClassDescriptor() = default;
-	};
+UE_DEPRECATED(5.3, "Including FieldNotification/IClassDescriptor.h is deprecated. The file location changed. Please include IFieldNotificationClassDescriptor.h instead and add the module FieldNotification to your .build.cs")
+inline void IClassDescriptorHeaderDeprecatedWarning()
+{
+}
 
-} // namespace
\ No newline at end of file
+inline void TriggerIClassDescriptorHeaderDeprecatedWarning()
+{
+	IClassDescriptorHeaderDeprecatedWarning();
+}
diff --git a/Engine/Source/Runtime/UMG/Public/FieldNotification/IFieldValueChanged.h b/Engine/Source/Runtime/UMG/Public/FieldNotification/IFieldValueChanged.h
--- a/Engine/Source/Runtime/UMG/Public/FieldNotification/IFieldValueChanged.h
+++ b/Engine/Source/Runtime/UMG/Public/FieldNotification/IFieldValueChanged.h
@@ -1,46 +1,16 @@
 // Copyright Epic Games, Inc. All Rights Reserved.
-
 #pragma once
 
-#include "CoreMinimal.h"
-#include "FieldNotification/FieldId.h"
-#include "FieldNotification/IClassDescriptor.h"
-#include "UObject/Interface.h"
-
-#include "IFieldValueChanged.generated.h"
-
+// HEADER_UNIT_SKIP - Deprecated
 
-DECLARE_DYNAMIC_DELEGATE_TwoParams(FFieldValueChangedDynamicDelegate, UObject*, Object, FFieldNotificationId, Field);
+#include "INotifyFieldValueChanged.h"
 
-
-UINTERFACE(MinimalAPI, NotBlueprintable)
-class UNotifyFieldValueChanged : public UInterface
+UE_DEPRECATED(5.3, "Including FieldNotification/IFieldValueChanged.h is deprecated. The file location changed. Please include INotifyFieldValueChanged.h instead and add the module FieldNotification to your .build.cs")
+inline void IFieldValueChangedHeaderDeprecatedWarning()
 {
-	GENERATED_BODY()
-};
-
+}
 
-class INotifyFieldValueChanged : public IInterface
+inline void TriggerIFieldValueChangedHeaderDeprecatedWarning()
 {
-	GENERATED_BODY()
-
-public:
-	DECLARE_MULTICAST_DELEGATE_TwoParams(FFieldValueChangedMulticastDelegate, UObject*, UE::FieldNotification::FFieldId);
-	using FFieldValueChangedDelegate = FFieldValueChangedMulticastDelegate::FDelegate;
-
-public:
-	/** Add a delegate that will be notified when the FieldId is value changed. */
-	virtual FDelegateHandle AddFieldValueChangedDelegate(UE::FieldNotification::FFieldId InFieldId, FFieldValueChangedDelegate InNewDelegate) = 0;
-
-	/** Remove a delegate that was added. */
-	virtual bool RemoveFieldValueChangedDelegate(UE::FieldNotification::FFieldId InFieldId, FDelegateHandle InHandle) = 0;
-
-	/** Remove all the delegate that are bound to the specified UserObject. */
-	virtual int32 RemoveAllFieldValueChangedDelegates(const void* InUserObject) = 0;
-
-	/** Remove all the delegate that are bound to the specified Field and UserObject. */
-	virtual int32 RemoveAllFieldValueChangedDelegates(UE::FieldNotification::FFieldId InFieldId, const void* InUserObject) = 0;
-
-	/** @returns the list of all the field that can notify when their value changes. */
-	virtual const UE::FieldNotification::IClassDescriptor& GetFieldNotificationDescriptor() const = 0;
-};
+	IFieldValueChangedHeaderDeprecatedWarning();
+}
diff --git a/Engine/Source/Runtime/UMG/UMG.Build.cs b/Engine/Source/Runtime/UMG/UMG.Build.cs
--- a/Engine/Source/Runtime/UMG/UMG.Build.cs
+++ b/Engine/Source/Runtime/UMG/UMG.Build.cs
@@ -11,8 +11,8 @@ public class UMG : ModuleRules
 				"Core",
 				"CoreUObject",
 				"DeveloperSettings",
-                "Engine",
-                "InputCore",
+				"Engine",
+				"InputCore",
 				"Slate",
 				"SlateCore",
 				"RenderCore",
@@ -21,32 +21,33 @@ public class UMG : ModuleRules
 			}
 		);
 
-        PublicDependencyModuleNames.AddRange(
-            new string[] {
+		PublicDependencyModuleNames.AddRange(
+			new string[] {
+				"FieldNotification",
 				"HTTP",
 				"MovieScene",
-                "MovieSceneTracks",
-                "PropertyPath",
+				"MovieSceneTracks",
+				"PropertyPath",
 				"TimeManagement"
 			}
-        );
+		);
 
-        PrivateIncludePathModuleNames.AddRange(
-            new string[] {
+		PrivateIncludePathModuleNames.AddRange(
+			new string[] {
 				"SlateRHIRenderer",
 				"ImageWrapper",
-                "TargetPlatform",
+				"TargetPlatform",
 			}
-        );
+		);
 
 		if (Target.Type != TargetType.Server)
 		{
-            DynamicallyLoadedModuleNames.AddRange(
-                new string[] {
-				    "ImageWrapper",
-				    "SlateRHIRenderer",
-			    }
-            );
+			DynamicallyLoadedModuleNames.AddRange(
+				new string[] {
+					"ImageWrapper",
+					"SlateRHIRenderer",
+				}
+			);
 		}
 	}
 }
