Added TDelegate::IsCompatable() and made FFieldMulticastDelegate use it instead of using internal functions.

diff --git a/Engine/Source/Runtime/Core/Public/Delegates/DelegateBase.h b/Engine/Source/Runtime/Core/Public/Delegates/DelegateBase.h
--- a/Engine/Source/Runtime/Core/Public/Delegates/DelegateBase.h
+++ b/Engine/Source/Runtime/Core/Public/Delegates/DelegateBase.h
@@ -225,6 +225,20 @@ public:
 		return DelegateInstance && DelegateInstance->HasSameObject(InUserObject);
 	}
 
+	/** 
+	 * Checks to see if this delegate can ever become valid again - if not, it can
+	 * be removed from broadcast lists or otherwise repurposed.
+	 *
+	 * @return True if the delegate is compatable, false otherwise.
+	 */
+	FORCEINLINE bool IsCompactable() const
+	{
+		FReadAccessScope ReadScope = GetReadAccessScope();
+
+		const IDelegateInstance* DelegateInstance = GetDelegateInstanceProtected();
+		return !DelegateInstance || DelegateInstance->IsCompactable();
+	}
+
 	/**
 	 * Gets a handle to the delegate.
 	 *
diff --git a/Engine/Source/Runtime/FieldNotification/Private/FieldNotificationDelegate.cpp b/Engine/Source/Runtime/FieldNotification/Private/FieldNotificationDelegate.cpp
--- a/Engine/Source/Runtime/FieldNotification/Private/FieldNotificationDelegate.cpp
+++ b/Engine/Source/Runtime/FieldNotification/Private/FieldNotificationDelegate.cpp
@@ -152,13 +152,10 @@ FFieldMulticastDelegate::FRemoveResult FFieldMulticastDelegate::Remove(const FDy
 		FInvocationElement& Element = Delegates[Index];
 		if (Element.Key.DynamicName == InDynamicDelegate.GetFunctionName())
 		{
-			if (const IDelegateInstance* DelegateInstance = GetDelegateInstance(Element.Delegate))
+			if (Element.Delegate.IsBoundToObject(InDynamicDelegate.GetUObject()))
 			{
-				if (DelegateInstance->HasSameObject(InDynamicDelegate.GetUObject()))
-				{
-					RemoveElement(Element, Index, Result);
-					break;
-				}
+				RemoveElement(Element, Index, Result);
+				break;
 			}
 		}
 	}
@@ -264,22 +261,19 @@ FFieldMulticastDelegate::FRemoveFromResult FFieldMulticastDelegate::RemoveFrom(c
 	{
 		if (Element.Key.DynamicName == InDynamicDelegate.GetFunctionName())
 		{
-			if (const IDelegateInstance* DelegateInstance = Self->GetDelegateInstance(Element.Delegate))
+			if (Element.Delegate.IsBoundToObject(InDynamicDelegate.GetUObject()))
 			{
-				if (DelegateInstance->HasSameObject(InDynamicDelegate.GetUObject()))
+				if (Self->DelegateLockCount > 0)
 				{
-					if (Self->DelegateLockCount > 0)
-					{
-						Element.Delegate.Unbind();
-						++Self->CompactionCount;
-					}
-					else
-					{
-						Self->Delegates.RemoveAt(Index);
-					}
-					bRemoved = true;
-					return !bFieldPresent;
+					Element.Delegate.Unbind();
+					++Self->CompactionCount;
 				}
+				else
+				{
+					Self->Delegates.RemoveAt(Index);
+				}
+				bRemoved = true;
+				return !bFieldPresent;
 			}
 		}
 		if (Element.Key.Object == InObject && Element.Delegate.IsBound())
@@ -310,20 +304,17 @@ FFieldMulticastDelegate::FRemoveAllResult FFieldMulticastDelegate::RemoveAll(con
 	{
 		for (FInvocationElement& Element : Delegates)
 		{
-			if (const IDelegateInstance* DelegateInstance = GetDelegateInstance(Element.Delegate))
+			if (Element.Key.Object == InObject)
 			{
-				if (Element.Key.Object == InObject)
+				if (Element.Delegate.IsBoundToObject(InUserObject))
 				{
-					if (DelegateInstance->HasSameObject(InUserObject))
-					{
-						Element.Delegate.Unbind();
-						++CompactionCount;
-						++Result.RemoveCount;
-					}
-					else if (!DelegateInstance->IsCompactable())
-					{
-						SetHasFields(Element);
-					}
+					Element.Delegate.Unbind();
+					++CompactionCount;
+					++Result.RemoveCount;
+				}
+				else if (!Element.Delegate.IsCompactable())
+				{
+					SetHasFields(Element);
 				}
 			}
 		}
@@ -333,14 +324,13 @@ FFieldMulticastDelegate::FRemoveAllResult FFieldMulticastDelegate::RemoveAll(con
 		for (int32 Index = Delegates.Num() - 1; Index >= 0; --Index)
 		{
 			FInvocationElement& Element = Delegates[Index];
-			const IDelegateInstance* DelegateInstance = GetDelegateInstance(Element.Delegate);
-			if (DelegateInstance == nullptr || DelegateInstance->IsCompactable())
+			if (Element.Delegate.IsCompactable())
 			{
 				Delegates.RemoveAt(Index);
 			}
 			else if (Element.Key.Object == InObject)
 			{
-				if (DelegateInstance->HasSameObject(InUserObject))
+				if (Element.Delegate.IsBoundToObject(InUserObject))
 				{
 					Delegates.RemoveAt(Index);
 					++Result.RemoveCount;
@@ -372,20 +362,17 @@ FFieldMulticastDelegate::FRemoveAllResult FFieldMulticastDelegate::RemoveAll(con
 	{
 		for (FInvocationElement& Element : Delegates)
 		{
-			if (const IDelegateInstance* DelegateInstance = GetDelegateInstance(Element.Delegate))
+			if (Element.Key.Id == InFieldId && Element.Key.Object == InObject)
 			{
-				if (Element.Key.Id == InFieldId && Element.Key.Object == InObject)
+				if (Element.Delegate.IsBoundToObject(InUserObject))
 				{
-					if (DelegateInstance->HasSameObject(InUserObject))
-					{
-						Element.Delegate.Unbind();
-						++CompactionCount;
-						++Result.RemoveCount;
-					}
-					else if (!DelegateInstance->IsCompactable())
-					{
-						SetHasFields(Element);
-					}
+					Element.Delegate.Unbind();
+					++CompactionCount;
+					++Result.RemoveCount;
+				}
+				else if (!Element.Delegate.IsCompactable())
+				{
+					SetHasFields(Element);
 				}
 			}
 		}
@@ -395,14 +382,13 @@ FFieldMulticastDelegate::FRemoveAllResult FFieldMulticastDelegate::RemoveAll(con
 		for (int32 Index = Delegates.Num() - 1; Index >= 0; --Index)
 		{
 			FInvocationElement& Element = Delegates[Index];
-			const IDelegateInstance* DelegateInstance = GetDelegateInstance(Element.Delegate);
-			if (DelegateInstance == nullptr || DelegateInstance->IsCompactable())
+			if (Element.Delegate.IsCompactable())
 			{
 				Delegates.RemoveAt(Index);
 			}
 			else if (Element.Key.Id == InFieldId && Element.Key.Object == InObject)
 			{
-				if (DelegateInstance->HasSameObject(InUserObject))
+				if (Element.Delegate.IsBoundToObject(InUserObject))
 				{
 					Delegates.RemoveAt(Index);
 					++Result.RemoveCount;
@@ -487,8 +473,7 @@ void FFieldMulticastDelegate::ExecuteLockOperations()
 			for (int32 Index = Delegates.Num() - 1; Index >= 0; --Index)
 			{
 				FInvocationElement& Element = Delegates[Index];
-				const IDelegateInstance* DelegateInstance = GetDelegateInstance(Element.Delegate);
-				if (DelegateInstance == nullptr || DelegateInstance->IsCompactable())
+				if (Element.Delegate.IsCompactable())
 				{
 					Delegates.RemoveAt(Index);
 				}
diff --git a/Engine/Source/Runtime/FieldNotification/Public/FieldNotificationDelegate.h b/Engine/Source/Runtime/FieldNotification/Public/FieldNotificationDelegate.h
--- a/Engine/Source/Runtime/FieldNotification/Public/FieldNotificationDelegate.h
+++ b/Engine/Source/Runtime/FieldNotification/Public/FieldNotificationDelegate.h
@@ -134,11 +134,6 @@ private:
 		return AlgoImpl::UpperBoundInternal(GetData(Delegates), Num, InKey, FIdentityFunctor(), TLess<>());
 	}
 
-	const IDelegateInstance* GetDelegateInstance(FDelegate& InDelegate) const
-	{
-		return static_cast<FInstanceProtectedDelegate&>(static_cast<TDelegateBase<FDefaultDelegateUserPolicy>&>(InDelegate)).GetDelegateInstanceProtected();
-	}
-
 	void ExecuteLockOperations();
 	void RemoveElement(FInvocationElement& Element, int32 Index, FRemoveResult& Result);
 	void CompleteRemove(FRemoveResult& Result) const;
