Removed unused FFieldMulticastDelegate::FInstanceProtectedDelegate.

diff --git a/Engine/Source/Runtime/FieldNotification/Public/FieldNotificationDelegate.h b/Engine/Source/Runtime/FieldNotification/Public/FieldNotificationDelegate.h
--- a/Engine/Source/Runtime/FieldNotification/Public/FieldNotificationDelegate.h
+++ b/Engine/Source/Runtime/FieldNotification/Public/FieldNotificationDelegate.h
@@ -29,12 +29,6 @@ public:
 	using FDynamicDelegate = FFieldValueChangedDynamicDelegate;
 
 private:
-	struct FInstanceProtectedDelegate : public TDelegateBase<FDefaultDelegateUserPolicy>
-	{
-		using Super = TDelegateBase<FDefaultDelegateUserPolicy>;
-		using Super::GetDelegateInstanceProtected;
-	};
-
 	struct FInvocationKey
 	{
 		TWeakObjectPtr<const UObject> Object;
