MVVM: Change the delegate type to not consider the threadsafe detector. The delegate is saved in a tarray and it can move or remove element from the tarray creating unwanted copies.

diff --git a/Engine/Source/Runtime/FieldNotification/Public/INotifyFieldValueChanged.h b/Engine/Source/Runtime/FieldNotification/Public/INotifyFieldValueChanged.h
--- a/Engine/Source/Runtime/FieldNotification/Public/INotifyFieldValueChanged.h
+++ b/Engine/Source/Runtime/FieldNotification/Public/INotifyFieldValueChanged.h
@@ -24,8 +24,9 @@ class INotifyFieldValueChanged : public IInterface
 	GENERATED_BODY()
 
 public:
-	DECLARE_MULTICAST_DELEGATE_TwoParams(FFieldValueChangedMulticastDelegate, UObject*, UE::FieldNotification::FFieldId);
-	using FFieldValueChangedDelegate = FFieldValueChangedMulticastDelegate::FDelegate;
+	// using "not checked" user policy (means race detection is disabled) because this delegate is stored in a container and causes its reallocation
+	// from inside delegate's execution. This is incompatible with race detection that needs to access the delegate instance after its execution
+	using FFieldValueChangedDelegate = TDelegate<void(UObject*, UE::FieldNotification::FFieldId), FNotThreadSafeNotCheckedDelegateUserPolicy>;
 
 public:
 	/** Add a delegate that will be notified when the FieldId is value changed. */
