[UHT/UMG]
* Fixed so dllexport is on methods/staticvars instead of type for FieldNotificationDeclarations

diff --git a/Engine/Source/Programs/Shared/EpicGames.UHT/Exporters/CodeGen/UhtHeaderCodeGeneratorHFile.cs b/Engine/Source/Programs/Shared/EpicGames.UHT/Exporters/CodeGen/UhtHeaderCodeGeneratorHFile.cs
--- a/Engine/Source/Programs/Shared/EpicGames.UHT/Exporters/CodeGen/UhtHeaderCodeGeneratorHFile.cs
+++ b/Engine/Source/Programs/Shared/EpicGames.UHT/Exporters/CodeGen/UhtHeaderCodeGeneratorHFile.cs
@@ -1007,7 +1007,7 @@ namespace EpicGames.UHT.Exporters.CodeGen
 					// UE_FIELD_NOTIFICATION_DECLARE_FIELD
 					foreach (UhtType notifyType in filteredTypes)
 					{
-						builder.Append("\tUE_FIELD_NOTIFICATION_DECLARE_FIELD(").Append(GetNotifyTypeName(notifyType)).Append(") \\\r\n");
+						builder.Append("\tUE_FIELD_NOTIFICATION_DECLARE_FIELD(").Append(GetNotifyTypeName(notifyType)).Append(',').Append(Module.Api).Append(") \\\r\n");
 					}
 
 					// UE_FIELD_NOTIFICATION_DECLARE_ENUM_FIELD
diff --git a/Engine/Source/Runtime/FieldNotification/Public/FieldNotificationDeclaration.h b/Engine/Source/Runtime/FieldNotification/Public/FieldNotificationDeclaration.h
--- a/Engine/Source/Runtime/FieldNotification/Public/FieldNotificationDeclaration.h
+++ b/Engine/Source/Runtime/FieldNotification/Public/FieldNotificationDeclaration.h
@@ -35,11 +35,11 @@
  	}
 
 #define UE_FIELD_NOTIFICATION_DECLARE_CLASS_DESCRIPTOR_BEGIN(API_STRING) \
-	struct API_STRING FFieldNotificationClassDescriptor : public Super::FFieldNotificationClassDescriptor \
+	struct FFieldNotificationClassDescriptor : public Super::FFieldNotificationClassDescriptor \
 	{ \
 	private: \
 		using SuperDescriptor = Super::FFieldNotificationClassDescriptor; \
-		static const ::UE::FieldNotification::FFieldId* AllFields[]; \
+		static API_STRING const ::UE::FieldNotification::FFieldId* AllFields[]; \
 		friend ThisClass; \
 	public:
 
@@ -61,17 +61,17 @@
 
 
 #define UE_FIELD_NOTIFICATION_DECLARE_CLASS_DESCRIPTOR_BASE_BEGIN(API_STRING) \
-	struct API_STRING FFieldNotificationClassDescriptor : public ::UE::FieldNotification::IClassDescriptor \
+	struct FFieldNotificationClassDescriptor : public ::UE::FieldNotification::IClassDescriptor \
 	{ \
 	private: \
 		using SuperDescriptor = ::UE::FieldNotification::IClassDescriptor; \
-		static const ::UE::FieldNotification::FFieldId* AllFields[]; \
+		static API_STRING const ::UE::FieldNotification::FFieldId* AllFields[]; \
 		friend ThisClass; \
-	public:
+	public: \
+		API_STRING virtual void ForEachField(const UClass* Class, TFunctionRef<bool(::UE::FieldNotification::FFieldId FielId)> Callback) const override; \
 
 
 #define UE_FIELD_NOTIFICATION_DECLARE_CLASS_DESCRIPTOR_BASE_END() \
-		virtual void ForEachField(const UClass* Class, TFunctionRef<bool(::UE::FieldNotification::FFieldId FielId)> Callback) const override; \
 	}; \
 	UE_FIELD_NOTIFICATION_OVERRIDE_GET_FIELD_NOTIFICATION_DESCRIPTOR()
 
@@ -87,8 +87,8 @@
 		};
 
 
-#define UE_FIELD_NOTIFICATION_DECLARE_FIELD(Name) \
-	static const ::UE::FieldNotification::FFieldId Name;
+#define UE_FIELD_NOTIFICATION_DECLARE_FIELD(Name, API_STRING) \
+	static API_STRING const ::UE::FieldNotification::FFieldId Name;
 
 
 #define UE_FIELD_NOTIFICATION_DECLARE_ENUM_FIELD(Name) \
diff --git a/Engine/Source/Runtime/UMG/Public/Components/Widget.h b/Engine/Source/Runtime/UMG/Public/Components/Widget.h
--- a/Engine/Source/Runtime/UMG/Public/Components/Widget.h
+++ b/Engine/Source/Runtime/UMG/Public/Components/Widget.h
@@ -220,9 +220,9 @@ class UWidget : public UVisual, public INotifyFieldValueChanged
 
 public:
 	UE_FIELD_NOTIFICATION_DECLARE_CLASS_DESCRIPTOR_BASE_BEGIN(UMG_API)
-		UE_FIELD_NOTIFICATION_DECLARE_FIELD(ToolTipText)
-		UE_FIELD_NOTIFICATION_DECLARE_FIELD(Visibility)
-		UE_FIELD_NOTIFICATION_DECLARE_FIELD(bIsEnabled)
+		UE_FIELD_NOTIFICATION_DECLARE_FIELD(ToolTipText, UMG_API)
+		UE_FIELD_NOTIFICATION_DECLARE_FIELD(Visibility, UMG_API)
+		UE_FIELD_NOTIFICATION_DECLARE_FIELD(bIsEnabled, UMG_API)
 		UE_FIELD_NOTIFICATION_DECLARE_ENUM_FIELD_BEGIN(ToolTipText)
 		UE_FIELD_NOTIFICATION_DECLARE_ENUM_FIELD(Visibility)
 		UE_FIELD_NOTIFICATION_DECLARE_ENUM_FIELD(bIsEnabled)
