[FieldNotification]
* Moved singleton to cpp file

diff --git a/Engine/Source/Runtime/FieldNotification/Public/FieldNotificationDeclaration.h b/Engine/Source/Runtime/FieldNotification/Public/FieldNotificationDeclaration.h
--- a/Engine/Source/Runtime/FieldNotification/Public/FieldNotificationDeclaration.h
+++ b/Engine/Source/Runtime/FieldNotification/Public/FieldNotificationDeclaration.h
@@ -30,8 +30,7 @@
 #define UE_FIELD_NOTIFICATION_OVERRIDE_GET_FIELD_NOTIFICATION_DESCRIPTOR() \
 	virtual const ::UE::FieldNotification::IClassDescriptor& GetFieldNotificationDescriptor() const \
 	{ \
- 		static FFieldNotificationClassDescriptor Instance; \
- 		return Instance; \
+ 		return FFieldNotificationClassDescriptor::GetDescriptor(); \
  	}
 
 #define UE_FIELD_NOTIFICATION_DECLARE_CLASS_DESCRIPTOR_BEGIN(API_STRING) \
@@ -40,6 +39,7 @@
 	private: \
 		using SuperDescriptor = Super::FFieldNotificationClassDescriptor; \
 		static API_STRING const ::UE::FieldNotification::FFieldId* AllFields[]; \
+		static API_STRING const ::UE::FieldNotification::IClassDescriptor& GetDescriptor(); \
 		friend ThisClass; \
 	public:
 
@@ -66,6 +66,7 @@
 	private: \
 		using SuperDescriptor = ::UE::FieldNotification::IClassDescriptor; \
 		static API_STRING const ::UE::FieldNotification::FFieldId* AllFields[]; \
+		static API_STRING const ::UE::FieldNotification::IClassDescriptor& GetDescriptor(); \
 		friend ThisClass; \
 	public: \
 		API_STRING virtual void ForEachField(const UClass* Class, TFunctionRef<bool(::UE::FieldNotification::FFieldId FielId)> Callback) const override; \
@@ -96,6 +97,11 @@
 
 
 #define UE_FIELD_NOTIFICATION_IMPLEMENTATION_BEGIN(ClassName) \
+	const ::UE::FieldNotification::IClassDescriptor& ClassName::FFieldNotificationClassDescriptor::GetDescriptor() \
+	{ \
+ 		static FFieldNotificationClassDescriptor Instance; \
+ 		return Instance; \
+ 	} \
 	const ::UE::FieldNotification::FFieldId* ClassName::FFieldNotificationClassDescriptor::AllFields[] = {
 
 
